Rename UnrealBuildToolTests to UnrealBuildTool.Tests, and include in the solution by default.

diff --git a/Engine/Source/Programs/UnrealBuildTool.Tests/PreprocessorExpressionTests.cs b/Engine/Source/Programs/UnrealBuildTool.Tests/PreprocessorExpressionTests.cs
new file mode 100644
--- /dev/null
+++ b/Engine/Source/Programs/UnrealBuildTool.Tests/PreprocessorExpressionTests.cs
@@ -0,0 +1,152 @@
+// Copyright Epic Games, Inc. All Rights Reserved.
+
+using Microsoft.VisualStudio.TestTools.UnitTesting;
+using System;
+using System.Collections.Generic;
+using System.Linq;
+using System.Text;
+using System.Threading.Tasks;
+using UnrealBuildTool;
+
+namespace UnrealBuildToolTests
+{
+	[TestClass]
+	public class PreprocessorExpressionTests
+	{
+		[TestMethod]
+		public void Addition()
+		{
+			RunTest("1 + 2 > 3", 0);
+			RunTest("1 + 2 >= 3", 1);
+		}
+
+		[TestMethod]
+		public void Equality()
+		{
+			RunTest("1 == 1", 1);
+			RunTest("1 == 2", 0);
+			RunTest("1 != 1", 0);
+			RunTest("1 != 2", 1);
+			RunTest("1 < 0", 0);
+			RunTest("1 < 1", 0);
+			RunTest("1 < 2", 1);
+			RunTest("1 <= 0", 0);
+			RunTest("1 <= 1", 1);
+			RunTest("1 <= 2", 1);
+			RunTest("1 > 0", 1);
+			RunTest("1 > 1", 0);
+			RunTest("1 > 2", 0);
+			RunTest("1 >= 0", 1);
+			RunTest("1 >= 1", 1);
+			RunTest("1 >= 2", 0);
+		}
+
+		[TestMethod]
+		public void UnaryOperators()
+		{
+			RunTest("0 + 1 == +1", 1);
+			RunTest("0 - 1 == -1", 1);
+			RunTest("!0", 1);
+			RunTest("!1", 0);
+			RunTest("~0", -1L);
+		}
+
+		[TestMethod]
+		public void ArithmeticOperators()
+		{
+			RunTest("3 + 7", 10);
+			RunTest("3 - 7", -4);
+			RunTest("3 * 7", 21);
+			RunTest("21 / 3", 7);
+			RunTest("24 % 7", 3);
+			RunTest("2 << 4", 32);
+			RunTest("32 >> 4", 2);
+			RunTest("(0xab & 0xf)", 0xb);
+			RunTest("(0xab | 0xf)", 0xaf);
+			RunTest("(0xab ^ 0xf)", 0xa4);
+		}
+
+		[TestMethod]
+		public void LogicalOperators()
+		{
+			RunTest("0 || 0", 0);
+			RunTest("0 || 1", 1);
+			RunTest("1 || 0", 1);
+			RunTest("1 || 1", 1);
+			RunTest("0 && 0", 0);
+			RunTest("0 && 1", 0);
+			RunTest("1 && 0", 0);
+			RunTest("1 && 1", 1);
+		}
+
+		[TestMethod]
+		public void TernaryOperator()
+		{
+			RunTest("((0)? 123 : 456) == 456", 1);
+			RunTest("((1)? 123 : 456) == 456", 0);
+		}
+
+		[TestMethod]
+		public void Precedence()
+		{
+			RunTest("3 + 27 / 3", 12);
+			RunTest("(3 + 27) / 3", 10);
+		}
+
+		[TestMethod]
+		public void Literals()
+		{
+			RunTest("123", 123);
+			RunTest("123l", 123);
+			RunTest("123L", 123);
+			RunTest("123ll", 123);
+			RunTest("123LL", 123);
+			RunTest("123u", 123);
+			RunTest("123U", 123);
+			RunTest("123ul", 123);
+			RunTest("123Ul", 123);
+			RunTest("123UL", 123);
+			RunTest("123Ull", 123);
+			RunTest("123ULL", 123);
+			RunTest("123ui64", 123);
+			RunTest("0x123i64", 0x123);
+			RunTest("0123i64", 83);
+        }
+
+		class PreprocessorTestContext : PreprocessorContext
+		{
+			public PreprocessorTestContext()
+				: base(null)
+			{
+			}
+		}
+
+		/// <summary>
+		/// Tokenize the given expression and evaluate it, and check it matches the expected result
+		/// </summary>
+		/// <param name="Expression">The expression to evaluate, as a string</param>
+		/// <param name="ExpectedResult">The expected value of the expression</param>
+        static void RunTest(string Expression, long? ExpectedResult)
+        {
+			TokenReader Reader = new TokenReader(Expression);
+
+			List<Token> Tokens = new List<Token>();
+			while (Reader.MoveNext())
+			{
+				Tokens.Add(Reader.Current);
+			}
+
+			long? Result;
+			try
+			{
+				Result = PreprocessorExpression.Evaluate(new PreprocessorTestContext(), Tokens);
+			}
+			catch(PreprocessorException)
+			{
+				Result = null;
+			}
+
+			Assert.AreEqual(ExpectedResult, Result);
+		}
+	}
+}
diff --git a/Engine/Source/Programs/UnrealBuildTool.Tests/PreprocessorStateTests.cs b/Engine/Source/Programs/UnrealBuildTool.Tests/PreprocessorStateTests.cs
new file mode 100644
--- /dev/null
+++ b/Engine/Source/Programs/UnrealBuildTool.Tests/PreprocessorStateTests.cs
@@ -0,0 +1,150 @@
+// Copyright Epic Games, Inc. All Rights Reserved.
+
+using Microsoft.VisualStudio.TestTools.UnitTesting;
+using System;
+using System.Collections.Generic;
+using System.Linq;
+using System.Text;
+using System.Threading.Tasks;
+using EpicGames.Core;
+using UnrealBuildTool;
+
+namespace UnrealBuildToolTests
+{
+	[TestClass]
+	public class PreprocessorStateTests
+	{
+		[TestMethod]
+		public void Run()
+		{
+			PreprocessorState BaseState = new PreprocessorState();
+			BaseState.DefineMacro(ParseMacro("FOO", null, "123"));
+			BaseState.DefineMacro(ParseMacro("BAR", new List<string>{ "Arg1", "Arg2" }, "Arg1 + Arg2"));
+			Assert.AreEqual("m BAR(Arg1, Arg2)=Arg1 + Arg2\nm FOO=123\n", FormatState(BaseState));
+
+			// Modify the state and revert it
+			{
+				PreprocessorState State = new PreprocessorState(BaseState);
+				Assert.AreEqual(FormatState(BaseState), FormatState(State));
+
+				PreprocessorTransform Transform1 = State.BeginCapture();
+				State.UndefMacro(Identifier.FindOrAdd("FOO"));
+				State.DefineMacro(ParseMacro("FOO2", null, "FOO()"));
+				State.PushBranch(PreprocessorBranch.Active);
+				Assert.AreEqual("b Active\nm BAR(Arg1, Arg2)=Arg1 + Arg2\nm FOO2=FOO()\n", FormatState(State));
+				Assert.AreEqual("+b Active\n+m FOO undef\n+m FOO2=FOO()\n", FormatTransform(Transform1));
+
+				PreprocessorState State2 = new PreprocessorState(BaseState);
+				State2.TryToApply(Transform1);
+				Assert.AreEqual(FormatState(State), FormatState(State2));
+
+				PreprocessorTransform Transform2 = State.BeginCapture();
+				State.PopBranch();
+				State.UndefMacro(Identifier.FindOrAdd("FOO"));
+				State.DefineMacro(ParseMacro("FOO", null, "123"));
+				State.UndefMacro(Identifier.FindOrAdd("FOO2"));
+				Assert.AreEqual(FormatState(BaseState), FormatState(State));
+				Assert.AreEqual("-b Active\n+m FOO=123\n+m FOO2 undef\n", FormatTransform(Transform2));
+
+				State2.TryToApply(Transform2);
+				Assert.AreEqual(FormatState(BaseState), FormatState(State2));
+			}
+
+			// Check the tracking of branches
+			{
+				PreprocessorState State = new PreprocessorState(BaseState);
+				PreprocessorTransform Transform1 = State.BeginCapture();
+				State.IsCurrentBranchActive();
+				Assert.AreEqual("=b Active\n", FormatTransform(Transform1));
+
+				State.PushBranch(PreprocessorBranch.Active);
+				PreprocessorTransform Transform2 = State.BeginCapture();
+				State.IsCurrentBranchActive();
+				State.PopBranch();
+				State.EndCapture();
+				Assert.AreEqual("-b Active\n", FormatTransform(Transform2));
+
+				State.PushBranch(0);
+				Assert.AreEqual("False", State.CanApply(Transform2).ToString());
+
+				State.PopBranch();
+				State.PushBranch(PreprocessorBranch.Active);
+				Assert.AreEqual("True", State.CanApply(Transform2).ToString());
+			}
+		}
+	
+		static string FormatState(PreprocessorState State)
+		{
+			StringBuilder Result = new StringBuilder();
+			foreach(PreprocessorBranch Branch in State.CurrentBranches)
+			{
+				Result.AppendFormat("b {0}\n", Branch);
+			}
+			foreach(PreprocessorMacro Macro in State.CurrentMacros.OrderBy(x => x.Name))
+			{
+				Result.AppendFormat("m {0}\n", Macro.ToString().TrimEnd());
+			}
+			return Result.ToString();
+		}
+
+		static string FormatTransform(PreprocessorTransform Transform)
+		{
+			StringBuilder Result = new StringBuilder();
+			if(Transform.bRequireTopmostActive.HasValue)
+			{
+				Result.AppendFormat("=b {0}\n", Transform.bRequireTopmostActive.Value? "Active" : "0");
+			}
+			foreach(PreprocessorBranch Branch in Transform.RequiredBranches)
+			{
+				Result.AppendFormat("-b {0}\n", Branch);
+			}
+			foreach(PreprocessorBranch Branch in Transform.NewBranches)
+			{
+				Result.AppendFormat("+b {0}\n", Branch);
+			}
+			foreach(KeyValuePair<Identifier, PreprocessorMacro> Macro in Transform.RequiredMacros)
+			{
+				if(Macro.Value == null)
+				{
+					Result.AppendFormat("=m {0} undef\n", Macro.Key);
+				}
+				else
+				{
+					Result.AppendFormat("=m {0}\n", Macro.Value.ToString().TrimEnd());
+				}
+			}
+			foreach(KeyValuePair<Identifier, PreprocessorMacro> Macro in Transform.NewMacros)
+			{
+				if(Macro.Value == null)
+				{
+					Result.AppendFormat("+m {0} undef\n", Macro.Key);
+				}
+				else
+				{
+					Result.AppendFormat("+m {0}\n", Macro.Value.ToString().TrimEnd());
+				}
+			}
+			return Result.ToString();
+		}
+
+		static PreprocessorMacro ParseMacro(string Name, List<string> Parameters, string Value)
+		{
+			List<Token> Tokens = new List<Token>();
+
+			TokenReader Reader = new TokenReader(Value);
+			while(Reader.MoveNext())
+			{
+				Tokens.Add(Reader.Current);
+			}
+
+			if(Parameters == null)
+			{
+				return new PreprocessorMacro(Identifier.FindOrAdd(Name), null, Tokens);
+			}
+			else
+			{
+				return new PreprocessorMacro(Identifier.FindOrAdd(Name), Parameters.Select(x => Identifier.FindOrAdd(x)).ToList(), Tokens);
+			}
+		}
+	}
+}
diff --git a/Engine/Source/Programs/UnrealBuildTool.Tests/PreprocessorTests.cs b/Engine/Source/Programs/UnrealBuildTool.Tests/PreprocessorTests.cs
new file mode 100644
--- /dev/null
+++ b/Engine/Source/Programs/UnrealBuildTool.Tests/PreprocessorTests.cs
@@ -0,0 +1,276 @@
+// Copyright Epic Games, Inc. All Rights Reserved.
+
+using Microsoft.VisualStudio.TestTools.UnitTesting;
+using System.Collections.Generic;
+using System.Text;
+using UnrealBuildTool;
+
+namespace UnrealBuildToolTests
+{
+	[TestClass]
+	public class PreprocessorTests
+	{
+		[TestMethod]
+		public void BasicTests()
+		{
+			RunTest("token", "token\n");
+			RunTest("#", "");
+		}
+
+		[TestMethod]
+		public void ObjectMacroTests()
+		{
+			RunTest("#define X token\nX", "token\n");
+			RunTest("#define A B C D\nA", "B C D\n");
+			RunTest("#define A B ## D\nA", "BD\n");
+
+			RunTest("#define X list of tokens\nX", "list of tokens\n");
+			RunTest("#define LST list\n#define TOKS tokens\n#define LOTS LST of TOKS\nLOTS LOTS", "list of tokens list of tokens\n");
+			RunTest("#define LOTS LST of TOKS\n#define LST list\n#define TOKS tokens\nLOTS LOTS", "list of tokens list of tokens\n");
+
+			RunTest("#define A MACRO\nA", "MACRO\n");
+			RunTest("#define A MACRO\n#undef A\nA", "A\n");
+		}
+
+		[TestMethod]
+		public void FunctionMacroTests()
+		{
+			RunTest("#define X(x) token x other\nX(and one).", "token and one other.\n");
+			RunTest("#define X(x,y) token x and y\nX(1, 2).", "token 1 and 2.\n");
+			RunTest("#define INC(x) (x + 2)\nINC", "INC\n");
+			RunTest("#define TEST(x) x\\\n?\nTEST(A)", "A?\n");
+
+			RunTest("#define F(A) <A>\nF(x)", "<x>\n");
+			RunTest("#define F(A,B) <A,B>\nF(x,y) + 1", "<x,y> + 1\n");
+			RunTest("#define F(A,B,C) <A,B,C>\nF(x,y,z)", "<x,y,z>\n");
+			RunTest("#define F(...) <__VA_ARGS__>\nF(x)", "<x>\n");
+			RunTest("#define F(...) <__VA_ARGS__>\nF(x,y)", "<x,y>\n");
+			RunTest("#define F(A,...) <A,__VA_ARGS__>\nF(x,y)", "<x,y>\n");
+			RunTest("#define F(A,...) <A, __VA_ARGS__>\nF(x,y)", "<x, y>\n");
+			RunTest("#define F(A,...) <A, __VA_ARGS__>\nF(x,y, z)", "<x, y, z>\n");
+
+			RunTest("#define FUNC(x) arg=x.\nFUNC(var) FUNC(2)", "arg=var. arg=2.\n");
+			RunTest("#define FUNC(x,y,z) int n = z+y*x;\nFUNC(1,2,3)", "int n = 3+2*1;\n");
+			RunTest("#define X 20\n#define FUNC(x,y) x+y\nx=FUNC(X,Y);", "x=20+Y;\n");
+			RunTest("#define FA(x,y) FB(x,y)\n#define FB(x,y) x + y\nFB(1,2);", "1 + 2;\n");
+			RunTest("#define PRINTF(...) printf(__VA_ARGS__)\nPRINTF()", "printf()\n");
+			RunTest("#define PRINTF(...) printf(__VA_ARGS__)\nPRINTF(\"hello\")", "printf(\"hello\")\n");
+			RunTest("#define PRINTF(...) printf(__VA_ARGS__)\nPRINTF(\"%d\", 1)", "printf(\"%d\", 1)\n");
+			RunTest("#define PRINTF(FORMAT, ...) printf(FORMAT, __VA_ARGS__)\nPRINTF(\"test\")", "printf(\"test\", )\n");
+			RunTest("#define PRINTF(FORMAT, ...) printf(FORMAT, __VA_ARGS__)\nPRINTF(\"test %s\", \"hello\")", "printf(\"test %s\", \"hello\")\n");
+			RunTest("#define PRINTF(FORMAT, ...) printf(FORMAT, __VA_ARGS__)\nPRINTF(\"test %s %d\", \"hello\", 1)", "printf(\"test %s %d\", \"hello\", 1)\n");
+			RunTest("#define PRINTF(FORMAT, ...) printf(FORMAT, ## __VA_ARGS__)\nPRINTF(\"test\")", "printf(\"test\" )\n");
+
+			RunTest("#define INC(x) (x + 2)\nINC(\n1\n)", "(1 + 2)\n");
+
+			RunTest("#define STRINGIZE(ARG) #ARG\nSTRINGIZE(+=)", "\"+=\"\n");
+			RunTest("#define STRINGIZE(ARG) #ARG\nSTRINGIZE(:>)", "\":>\"\n");
+			RunTest("#define STRINGIZE(ARG) #ARG\nSTRINGIZE(3.1415)", "\"3.1415\"\n");
+
+			RunTest("#define CONCAT(X, Y) X ## Y\nCONCAT(+, =)", "+=\n");
+			RunTest("#define CONCAT(X, Y) X ## Y\nCONCAT(3, .14159)", "3.14159\n");
+			RunTest("#define CONCAT(X, Y) X ## Y\nCONCAT(3, hello)", "3hello\n");
+			RunTest("#define CONCAT(X, Y) X ## #Y\nCONCAT(u, hello)", "u\"hello\"\n");
+			RunTest("#define CONCAT(X, ...) X ## __VA_ARGS__\nCONCAT(hello) there", "hello there\n");
+			RunTest("#define CONCAT(X, ...) X ## __VA_ARGS__\nCONCAT(hello, there)", "hellothere\n");
+
+			RunTest("#define A(X) MACRO X\nA(x)", "MACRO x\n");
+			RunTest("#define A(X) MACRO X\n#undef A\nA(x)", "A(x)\n");
+		}
+
+		[TestMethod]
+		public void ConditionalTests()
+		{
+			RunTest("#if 1\nfirst_branch\n#endif\nend", "first_branch\nend\n");
+			RunTest("#if 0\nfirst_branch\n#endif\nend", "end\n");
+			RunTest("#if 1\nbranch_1\n#else\nbranch_2\n#endif", "branch_1\n");
+			RunTest("#if 0\nbranch_1\n#else\nbranch_2\n#endif", "branch_2\n");
+			RunTest("#define A\n#ifdef A\nYes\n#endif", "Yes\n");
+			RunTest("#define B\n#ifdef A\nYes\n#endif\nNo", "No\n");
+			RunTest("#define A\n#ifndef A\nYes\n#endif\nNo", "No\n");
+			RunTest("#define B\n#ifndef A\nYes\n#endif", "Yes\n");
+			RunTest("#define A\n#undef A\n#ifdef A\nYes\n#endif\nNo", "No\n");
+			RunTest("#define A\n#undef A\n#ifndef A\nYes\n#endif", "Yes\n");
+			RunTest("#define A\n#ifdef A\nYes\n#else\nNo\n#endif", "Yes\n");
+			RunTest("#define B\n#ifdef A\nYes\n#else\nNo\n#endif", "No\n");
+		}
+
+		[TestMethod]
+		public void ExpressionTests()
+		{
+			RunTest("#if 1 + 2 > 3\nYes\n#endif\nNo", "No\n");
+			RunTest("#if 1 + 2 >= 3\nYes\n#endif", "Yes\n");
+			RunTest("#define ONE 1\n#define TWO 2\n#define PLUS(x, y) x + y\n#if PLUS(ONE, TWO) > 3\nYes\n#endif\nNo", "No\n");
+			RunTest("#define ONE 1\n#define TWO 2\n#define PLUS(x, y) x + y\n#if PLUS(ONE, TWO) >= 3\nYes\n#endif", "Yes\n");
+			RunTest("#define ONE 1\n#if defined ONE\nOne\n#elif defined TWO\nTwo\n#else\nThree\n#endif", "One\n");
+			RunTest("#define TWO 1\n#if defined ONE\nOne\n#elif defined TWO\nTwo\n#else\nThree\n#endif", "Two\n");
+			RunTest("#define ONE 0\n#if defined(ONE) + defined(TWO) >= 1\nYes\n#else\nNo\n#endif", "Yes\n");
+			RunTest("#define ONE 0\n#if defined(ONE) + defined(TWO) >= 2\nYes\n#else\nNo\n#endif", "No\n");
+			RunTest("#define ONE 0\n#define TWO\n#if defined(ONE) + defined(TWO) >= 2\nYes\n#else\nNo\n#endif", "Yes\n");
+		}
+
+		[TestMethod]
+		public void ConcatenationTests()
+		{
+			RunTest("#define FUNC(x) Value = x\n#define PP_JOIN(x, y) x ## y\n#define RESULT(x) PP_JOIN(FU, NC)(x)\nRESULT(1234)", "Value = 1234\n");
+			RunTest("#define VALUE 1234\n#define PP_JOIN(x, y) x ## y\n#define RESULT PP_JOIN(V, ALUE)\nRESULT", "1234\n");
+			RunTest("#define V 1\n#define ALUE 2\n#define VALUE 1234\n#define PP_JOIN(x, y) x ## y\n#define RESULT PP_JOIN(V, ALUE)\nRESULT", "1234\n");
+		}
+
+		[TestMethod]
+		public void VarargTests()
+		{
+			RunTest("#define FUNC(fmt,...) (fmt, ## __VA_ARGS__)\nFUNC(a)", "(a )\n");
+			RunTest("#define FUNC(fmt,...) (fmt, ## __VA_ARGS__)\nFUNC(a,b)", "(a,b)\n");
+			RunTest("#define FUNC(fmt,...) (fmt, ## __VA_ARGS__)\nFUNC(a,b )", "(a,b)\n");
+			RunTest("#define FUNC(fmt, ...) (fmt, ## __VA_ARGS__)\nFUNC(a)", "(a )\n");
+			RunTest("#define FUNC(fmt, ...) (fmt, ## __VA_ARGS__)\nFUNC(a,b)", "(a,b)\n");
+			RunTest("#define FUNC(fmt, ...) (fmt, ## __VA_ARGS__)\nFUNC(a,b )", "(a,b)\n");
+		}
+
+		[TestMethod]
+		public void EmptyTokenTests()
+		{
+			RunTest("#define EMPTY_TOKEN\n#define FUNC(_FuncName, _HType1, _HArg1) _FuncName(_HType1 _HArg1);\nFUNC(hello, EMPTY_TOKEN, int)", "hello( int);\n");
+
+			RunTest("#define EMPTY_TOKEN\n#define GCC_EXTENSION(...) 123 , ## __VA_ARGS__\nGCC_EXTENSION(EMPTY_TOKEN)", "123  \n");
+
+			RunTest("#define EMPTY_TOKEN\n#define FUNC(x) (x)\nFUNC(EMPTY_TOKEN A)", "( A)\n");
+			RunTest("#define EMPTY_TOKEN\n#define FUNC(x,y) (x y)\nFUNC(EMPTY_TOKEN,A)", "( A)\n");
+			//			RunTest("#define EMPTY_TOKEN\n#define FUNC(x) (x)\n#define FUNC_2(x,y) FUNC(x y)\nFUNC_2(EMPTY_TOKEN,A)", "( A)\n");
+			//			RunTest("#define EMPTY_TOKEN\n#define FUNC(x) (x EMPTY_TOKEN)\nFUNC(A)", "(A )\n");
+			//			RunTest("#define EMPTY_TOKEN\n#define FUNC(x,y) (x y)\nFUNC(A,)", "(A)\n");
+
+
+			RunTest("#define EMPTY\n#define FUNC(x) (x)\nFUNC(x    y)", "(x y)\n");
+			RunTest("#define EMPTY\n#define FUNC(x) (x)\nFUNC(x    EMPTY    y)", "(x  y)\n");
+			RunTest("#define EMPTY\n#define FUNC(x) (x)\nFUNC(x    EMPTY    y    EMPTY)", "(x  y )\n");
+			RunTest("#define EMPTY\n#define FUNC(x) (x)\nFUNC(EMPTY    x    EMPTY    y)", "( x  y)\n");
+			RunTest("#define EMPTY\n#define FUNC(x) (x)\nFUNC(    EMPTY    x    EMPTY    y)", "( x  y)\n");
+			RunTest("#define EMPTY\n#define FUNC(x) (x)\nFUNC(EMPTY    x    EMPTY    y    )", "( x  y)\n");
+
+			RunTest("#define EMPTY\n#define FUNC(x) (x)\n#define FUNC_2(x) FUNC(x)\nFUNC_2(x    y)", "(x y)\n");
+			RunTest("#define EMPTY\n#define FUNC(x) (x)\n#define FUNC_2(x) FUNC(x)\nFUNC_2(x    EMPTY    y)", "(x y)\n");
+			RunTest("#define EMPTY\n#define FUNC(x) (x)\n#define FUNC_2(x) FUNC(x)\nFUNC_2(x    EMPTY    y    EMPTY)", "(x y)\n");
+			RunTest("#define EMPTY\n#define FUNC(x) (x)\n#define FUNC_2(x) FUNC(x)\nFUNC_2(EMPTY    x    EMPTY    y)", "( x y)\n");
+			RunTest("#define EMPTY\n#define FUNC(x) (x)\n#define FUNC_2(x) FUNC(x)\nFUNC_2(    EMPTY    x    EMPTY    y)", "( x y)\n");
+			RunTest("#define EMPTY\n#define FUNC(x) (x)\n#define FUNC_2(x) FUNC(x)\nFUNC_2(EMPTY    x    EMPTY    y    )", "( x y)\n");
+
+			RunTest("#define EMPTY\n#define FUNC(x) ( x )\nFUNC(EMPTY EMPTY)", "(   )\n");
+			RunTest("#define EMPTY\n#define FUNC(x) ( x)\nFUNC(EMPTY EMPTY x)", "(   x)\n");
+			RunTest("#define EMPTY\n#define FUNC(x,y) ( x y)\n#define FUNC2(x,y) FUNC(x,y)\nFUNC2(EMPTY EMPTY EMPTY, x)", "(   x)\n");
+			RunTest("#define EMPTY\n#define DOUBLE_EMPTY EMPTY EMPTY\n#define FUNC(x) ( x)\nFUNC(DOUBLE_EMPTY x)", "(   x)\n");
+			RunTest("#define EMPTY\n#define DOUBLE_EMPTY EMPTY EMPTY\n#define FUNC(x,y) ( x y )\n#define FUNC2(x,y) FUNC(x,y)\nFUNC2(DOUBLE_EMPTY EMPTY, x)", "(   x )\n");
+			RunTest("#define EMPTY\n#define FUNC(x,y) ( x y )\nFUNC(EMPTY EMPTY EMPTY EMPTY EMPTY EMPTY, x)", "(       x )\n");
+		}
+
+		[TestMethod]
+		public void MiscTests()
+		{
+			RunTest("#define REGISTER_NAME(num,name) NAME_##name = num,\nREGISTER_NAME(201, TRUE)", "NAME_TRUE = 201,\n");
+			RunTest("#define FUNC_2(X) VALUE=X\n#define FUNC_N(X) FUNC_##X\n#define OBJECT FUNC_N(2)\nOBJECT(1234)", "VALUE=1234\n");
+			RunTest("#define GCC_EXTENSION(...) 123, ## __VA_ARGS__\nGCC_EXTENSION(456)", "123,456\n");
+
+			RunTest("#define _NON_MEMBER_CALL(FUNC, CV_REF_OPT) FUNC(X,CV_REF_OPT)\n#define _NON_MEMBER_CALL_CV(FUNC, REF_OPT) _NON_MEMBER_CALL(FUNC, REF_OPT) _NON_MEMBER_CALL(FUNC, const REF_OPT)\n#define _NON_MEMBER_CALL_CV_REF(FUNC) _NON_MEMBER_CALL_CV(FUNC, ) _NON_MEMBER_CALL_CV(FUNC, &)\n#define _IS_FUNCTION(X,CV_REF_OPT) (CV_REF_OPT)\n_NON_MEMBER_CALL_CV_REF(_IS_FUNCTION)", "() (const) (&) (const &)\n");
+
+			RunTest("#define TEXT(x) L ## x\n#define checkf(expr, format, ...) FDebug::AssertFailed(#expr, format, ##__VA_ARGS__)\ncheckf( true, TEXT( \"hello world\" ) );", "FDebug::AssertFailed(\"true\", L\"hello world\" );\n");
+		}
+
+		[TestMethod]
+		public void ErrorTests()
+		{
+			RunTest("#define", null);
+			RunTest("#define INC(x) (x + 2)\nINC(", null);
+			RunTest("#define +", null);
+			RunTest("#define A B __VA_ARGS__ D\nA", null);
+			RunTest("#define A ## B\nA", null);
+			RunTest("#define A B ##\nA", null);
+			RunTest("#define defined not_defined", null);
+
+			RunTest("#define F(A) <A>\nF(x,y) + 1", null);
+			RunTest("#define F(A,) <A>\nF(x)", null);
+			RunTest("#define F(A+) <A>\nF(x)", null);
+			RunTest("#define F(A,A) <A>\nF(x,y)", null);
+			RunTest("#define F(A,B) <A,B>\nF(x) + 1", null);
+			RunTest("#define F(A,B,) <A,B>\nF(x,y) + 1", null);
+//			RunTest("#define F(...) <__VA_ARGS__>\nF(x,)", null);
+			RunTest("#define F(A...) <A, __VA_ARGS__>\nF(x)", null);
+			RunTest("#define F(A...) <A, __VA_ARGS__>\nF(x,y)", null);
+			RunTest("#define F(A,__VA_ARGS__) <A, __VA_ARGS__>\nF(x,y)", null);
+			RunTest("#define F(A) #+\nF(x)", null);
+			RunTest("#define F(A) #B\nF(x)", null);
+			RunTest("#define F(A) ## A\nF(x)", null);
+			RunTest("#define F(A) A ##\nF(x)", null);
+			RunTest("#define F(A) <__VA_ARGS__>\nF(x)", null);
+
+			RunTest("#define INC(x\n)", null);
+			RunTest("#if 1\nbranch_1\n#else garbage\nbranch_2\n#endif\nend", null);
+//			RunTest("#if 1\nbranch_1\n#else\nbranch_2\n#endif garbage\nend", null);
+			RunTest("#if 1\nbranch_1\n#else\nbranch_2\n#elif 1\nbranch_3\n#endif\nend", null);
+			RunTest("#if 1\nbranch_1\n#else\nbranch_2\n#else\nbranch_3\n#endif", null);
+			RunTest("#if 0\nbranch_1\n#else\nbranch_2\n#else\nbranch_3\n#endif", null);
+			RunTest("#elif\nbranch_1\n#else\nbranch_2\n#endif\nend", null);
+			RunTest("#ifdef +\nbranch_1\n#else\nbranch_2\n#endif", null);
+			RunTest("#ifdef A 1\nbranch_1\n#else\nbranch_2\n#endif", null);
+			RunTest("#define A()\n#if A(\n)\nOK\n#endif", null);
+			RunTest("#define A MACRO\n#undef A B\nA", null);
+		}
+
+		/// <summary>
+		/// Preprocess a fragment of code, and check it results in an expected sequence of tokens
+		/// </summary>
+		/// <param name="Fragment">The code fragment to preprocess</param>
+		/// <param name="ExpectedResult">The expected sequence of tokens, as a string. Null to indicate that the input is invalid, and an exception is expected.</param>
+		static void RunTest(string Fragment, string ExpectedResult)
+		{
+			string[] Lines = Fragment.Split('\n');
+
+			SourceFile File = new SourceFile(null, TokenReader.GetNullTerminatedByteArray(Fragment));
+
+			Preprocessor Instance = new Preprocessor();
+
+			string Result;
+			try
+			{
+				List<Token> OutputTokens = new List<Token>();
+				for (int MarkupIdx = 0; MarkupIdx < File.Markup.Length; MarkupIdx++)
+				{
+					SourceFileMarkup Markup = File.Markup[MarkupIdx];
+					if (Markup.Type == SourceFileMarkupType.Text)
+					{
+						if (Instance.IsCurrentBranchActive())
+						{
+							StringBuilder SourceText = new StringBuilder();
+
+							int LastLineIndex = (MarkupIdx + 1 < File.Markup.Length) ? File.Markup[MarkupIdx + 1].LineNumber - 1 : Lines.Length;
+							for (int LineIndex = Markup.LineNumber - 1; LineIndex < LastLineIndex; LineIndex++)
+							{
+								SourceText.AppendLine(Lines[LineIndex]);
+							}
+
+							List<Token> Tokens = new List<Token>();
+
+							TokenReader Reader = new TokenReader(SourceText.ToString());
+							while (Reader.MoveNext())
+							{
+								Tokens.Add(Reader.Current);
+							}
+
+							Instance.ExpandMacros(Tokens, OutputTokens, false, null);
+						}
+					}
+					else
+					{
+						Instance.ParseMarkup(Markup.Type, Markup.Tokens, null);
+					}
+				}
+				Result = Token.Format(OutputTokens);
+			}
+			catch (PreprocessorException)
+			{
+				Result = null;
+			}
+			Assert.AreEqual(ExpectedResult, Result);
+		}
+	}
+}
diff --git a/Engine/Source/Programs/UnrealBuildTool.Tests/SourceFileTests.cs b/Engine/Source/Programs/UnrealBuildTool.Tests/SourceFileTests.cs
new file mode 100644
--- /dev/null
+++ b/Engine/Source/Programs/UnrealBuildTool.Tests/SourceFileTests.cs
@@ -0,0 +1,96 @@
+// Copyright Epic Games, Inc. All Rights Reserved.
+
+using Microsoft.VisualStudio.TestTools.UnitTesting;
+using System;
+using System.Collections.Concurrent;
+using System.Collections.Generic;
+using System.Linq;
+using System.Text;
+using System.Threading.Tasks;
+using EpicGames.Core;
+using OpenTracing.Util;
+using UnrealBuildBase;
+using UnrealBuildTool;
+
+namespace UnrealBuildToolTests
+{
+	/// <summary>
+	/// Tests for reading source file markup
+	/// </summary>
+	[TestClass]
+	public class SourceFileTests
+	{
+		[TestMethod]
+		public void Run()
+		{
+			// Test assemblies are built into a temporary directory outside of the unreal source tree, which prevents Unreal.FindRootDirectory()
+			// from being able to deduce the correct location - so override it.q
+			// The default working directory for this test is ...\Engine\Source\Programs\UnrealBuildToolTests\bin\[configuration]\netcoreapp3.1
+			Unreal.LocationOverride.RootDirectory = DirectoryReference.Combine(DirectoryReference.GetCurrentDirectory(), @"..\..\..\..\..\..\..");
+
+			List<DirectoryReference> BaseDirectories = new List<DirectoryReference>();
+			BaseDirectories.Add(DirectoryReference.Combine(UnrealBuildTool.UnrealBuildTool.EngineSourceDirectory, "Runtime"));
+			BaseDirectories.Add(DirectoryReference.Combine(UnrealBuildTool.UnrealBuildTool.EngineSourceDirectory, "Developer"));
+			BaseDirectories.Add(DirectoryReference.Combine(UnrealBuildTool.UnrealBuildTool.EngineSourceDirectory, "Editor"));
+
+			foreach(FileReference PluginFile in PluginsBase.EnumeratePlugins((FileReference)null))
+			{
+				DirectoryReference PluginSourceDir = DirectoryReference.Combine(PluginFile.Directory, "Source");
+				if(DirectoryReference.Exists(PluginSourceDir))
+				{
+					BaseDirectories.Add(PluginSourceDir);
+				}
+			}
+
+			ConcurrentBag<SourceFile> SourceFiles = new ConcurrentBag<SourceFile>();
+			using (GlobalTracer.Instance.BuildSpan("Scanning source files").StartActive())
+			{
+				using(ThreadPoolWorkQueue Queue = new ThreadPoolWorkQueue())
+				{
+					foreach(DirectoryReference BaseDirectory in BaseDirectories)
+					{
+						Queue.Enqueue(() => ParseSourceFiles(DirectoryItem.GetItemByDirectoryReference(BaseDirectory), SourceFiles, Queue));
+					}
+				}
+			}
+
+			Log.TraceLog("Read {0} source files", SourceFiles.Count);
+
+			FileReference TempDataFile = FileReference.Combine(Unreal.EngineDirectory, "Intermediate", "Temp", "SourceFileTests.bin");
+			DirectoryReference.CreateDirectory(TempDataFile.Directory);
+
+			using (GlobalTracer.Instance.BuildSpan("Writing source file data").StartActive())
+			{
+				using(BinaryArchiveWriter Writer = new BinaryArchiveWriter(TempDataFile))
+				{
+					Writer.WriteList(SourceFiles.ToList(), x => x.Write(Writer));
+				}
+			}
+
+			List<SourceFile> ReadSourceFiles = new List<SourceFile>();
+			using (GlobalTracer.Instance.BuildSpan("Reading source file data").StartActive())
+			{
+				using(BinaryArchiveReader Reader = new BinaryArchiveReader(TempDataFile))
+				{
+					ReadSourceFiles = Reader.ReadList(() => new SourceFile(Reader));
+				}
+			}
+		}
+
+		static void ParseSourceFiles(DirectoryItem Directory, ConcurrentBag<SourceFile> SourceFiles, ThreadPoolWorkQueue Queue)
+		{
+			foreach(DirectoryItem SubDirectory in Directory.EnumerateDirectories())
+			{
+				Queue.Enqueue(() => ParseSourceFiles(SubDirectory, SourceFiles, Queue));
+			}
+
+			foreach(FileItem File in Directory.EnumerateFiles())
+			{
+				if(File.HasExtension(".h") || File.HasExtension(".cpp"))
+				{
+					Queue.Enqueue(() => SourceFiles.Add(new SourceFile(File)));
+				}
+			}
+		}
+	}
+}
diff --git a/Engine/Source/Programs/UnrealBuildTool.Tests/TokenReaderTests.cs b/Engine/Source/Programs/UnrealBuildTool.Tests/TokenReaderTests.cs
new file mode 100644
--- /dev/null
+++ b/Engine/Source/Programs/UnrealBuildTool.Tests/TokenReaderTests.cs
@@ -0,0 +1,174 @@
+// Copyright Epic Games, Inc. All Rights Reserved.
+
+using Microsoft.VisualStudio.TestTools.UnitTesting;
+using System;
+using System.Collections.Generic;
+using System.Linq;
+using System.Text;
+using System.Threading.Tasks;
+using UnrealBuildTool;
+
+namespace UnrealBuildToolTests
+{
+	[TestClass]
+	public class TokenReaderTests
+	{
+		[TestMethod]
+		public void Run()
+		{
+			RunLexerTest("ABC//hello\nDEF", "Identifier(ABC), Newline, Identifier(DEF)");
+			RunLexerTest("ABC/*hello\nworld*/DEF", "Identifier(ABC), Identifier(DEF)");
+			RunLexerTest("ABC//hello\r\nDEF", "Identifier(ABC), Newline, Identifier(DEF)");
+			RunLexerTest("ABC123", "Identifier(ABC123)");
+			RunLexerTest("123ABC", "Number(123ABC)");
+			RunLexerTest("FOO(1,2,3,...)", "Identifier(FOO), LeftParen, Number(1), Comma, Number(2), Comma, Number(3), Comma, Ellipsis, RightParen");
+		}
+
+		[TestMethod]
+		public void SimpleTokens()
+		{
+			RunLexerTest("=<><<+", "Equals, CompareLess, CompareGreater, LeftShift, Plus");
+			RunLexerTest("+=+++++", "Plus, Equals, Plus, Plus, Plus, Plus, Plus");
+			RunLexerTest("/+++=", "Divide, Plus, Plus, Plus, Equals");
+			RunLexerTest("###", "HashHash, Hash");
+			RunLexerTest("():?", "LeftParen, RightParen, Colon, QuestionMark");
+			RunLexerTest("..*->->*", "Dot, Dot, Multiply, Minus, CompareGreater, Minus, CompareGreater, Multiply");
+			RunLexerTest("+-*/%", "Plus, Minus, Multiply, Divide, Modulo");
+			RunLexerTest("|^&~", "BitwiseOr, BitwiseXor, BitwiseAnd, BitwiseNot");
+			RunLexerTest("+=-=*=/=", "Plus, Equals, Minus, Equals, Multiply, Equals, Divide, Equals");
+			RunLexerTest("%=|=^=", "Modulo, Equals, BitwiseOr, Equals, BitwiseXor, Equals");
+			RunLexerTest("&&||!", "LogicalAnd, LogicalOr, LogicalNot");
+			RunLexerTest("=<><=>=", "Equals, CompareLess, CompareGreater, CompareLessOrEqual, CompareGreaterOrEqual");
+			RunLexerTest("==!=", "CompareEqual, CompareNotEqual");
+			RunLexerTest(".....", "Ellipsis, Dot, Dot");
+		}
+
+		[TestMethod]
+		public void IntegerLiterals()
+		{
+			RunLexerTest("123+", "Number(123), Plus");
+			RunLexerTest("01234567+", "Number(01234567), Plus");
+			RunLexerTest("0x123456789ABCDEFabcdef+", "Number(0x123456789ABCDEFabcdef), Plus");
+			RunLexerTest("123Ull+", "Number(123Ull), Plus");
+			RunLexerTest("123_hello+", "Number(123_hello), Plus");
+			RunLexerTest("01234567_hello+", "Number(01234567_hello), Plus");
+			RunLexerTest("0x123456789ABCDEFabcdef_hello+", "Number(0x123456789ABCDEFabcdef_hello), Plus");
+		}
+
+		[TestMethod]
+		public void FloatLiterals()
+		{
+			RunLexerTest("123.+", "Number(123.), Plus");
+			RunLexerTest("123e123+", "Number(123e123), Plus");
+			RunLexerTest("123e+123+", "Number(123e+123), Plus");
+			RunLexerTest(".123+", "Number(.123), Plus");
+			RunLexerTest(".123e123+", "Number(.123e123), Plus");
+			RunLexerTest(".123e+123+", "Number(.123e+123), Plus");
+			RunLexerTest("0.123+", "Number(0.123), Plus");
+			RunLexerTest("0.123e123+", "Number(0.123e123), Plus");
+			RunLexerTest("0.123e+123+", "Number(0.123e+123), Plus");
+			RunLexerTest("123.123+", "Number(123.123), Plus");
+			RunLexerTest("123.123e123+", "Number(123.123e123), Plus");
+			RunLexerTest("123.123e+123+", "Number(123.123e+123), Plus");
+			RunLexerTest("1.f+", "Number(1.f), Plus");
+			RunLexerTest("123.123e+123f+", "Number(123.123e+123f), Plus");
+			RunLexerTest("123.123e+123F+", "Number(123.123e+123F), Plus");
+		}
+
+		[TestMethod]
+		public void CharacterLiterals()
+		{
+			RunLexerTest("'hello'+", "Character('hello'), Plus");
+			RunLexerTest("'hello\\'\\''+", "Character('hello\\'\\''), Plus");
+			RunLexerTest("u'hello'+", "Character(u'hello'), Plus");
+			RunLexerTest("U'hello'+", "Character(U'hello'), Plus");
+			RunLexerTest("L'hello'+", "Character(L'hello'), Plus");
+			RunLexerTest("u8'hello'+", "Character(u8'hello'), Plus");
+		}
+
+		[TestMethod]
+		public void StringLiterals()
+		{
+			RunLexerTest("\"hello\"+", "String(\"hello\"), Plus");
+			RunLexerTest("\"hello\\\"\\\"\"+", "String(\"hello\\\"\\\"\"), Plus");
+			RunLexerTest("u\"hello\"+", "String(u\"hello\"), Plus");
+			RunLexerTest("U\"hello\"+", "String(U\"hello\"), Plus");
+			RunLexerTest("L\"hello\"+", "String(L\"hello\"), Plus");
+			RunLexerTest("u8\"hello\"+", "String(u8\"hello\"), Plus");
+		}
+
+		[TestMethod]
+		public void Identifiers()
+		{
+			RunLexerTest("abcdefghijklmnopqrstuvwxyz1234567890ABCDEFGHIJKLMNOPQRSTUVWXYZ_+", "Identifier(abcdefghijklmnopqrstuvwxyz1234567890ABCDEFGHIJKLMNOPQRSTUVWXYZ_), Plus");
+		}
+
+		[TestMethod]
+		public void Comments()
+		{
+			// Single-line comments
+			RunLexerTest("//hello\n+", "Newline, Plus");
+			RunLexerTest("//hello\n//world\n+", "Newline, Newline, Plus");
+
+			// Multi-line comments
+			RunLexerTest("/*hello*/+", "Plus");
+			RunLexerTest("/*/hello*/+", "Plus");
+			RunLexerTest("/*/hello**/+", "Plus");
+		}
+
+		[TestMethod]
+		public void LineNumbers()
+		{
+			// Test reported line numbers are accurate
+			RunLineNumberTest("+", 1);
+			RunLineNumberTest("\n+", 2);
+			RunLineNumberTest("\r\n+", 2);
+			RunLineNumberTest("\r\n\r\n+", 3);
+			RunLineNumberTest("\r\n\n+", 3);
+			RunLineNumberTest("//comment\n+", 2);
+			RunLineNumberTest("//comment\\\n\n+", 3);
+			RunLineNumberTest("/*comment\n*/+", 2);
+			RunLineNumberTest("/*comment\r\n*/+", 2);
+			RunLineNumberTest("/*comment\r\n\r\n*/+", 3);
+			RunLineNumberTest("first line \\\n+", 2);
+		}
+
+		public static void RunLexerTest(string InputText, string ExpectedOutputText)
+		{
+			List<Token> Tokens = new List<Token>();
+
+			TokenReader Reader = new TokenReader(InputText);
+			while(Reader.MoveNext())
+			{
+				Tokens.Add(Reader.Current);
+			}
+
+			string OutputText = String.Join(", ", Tokens.Select(x => FormatToken(x)));
+			Assert.AreEqual(ExpectedOutputText, OutputText);
+		}
+
+		static string FormatToken(Token Token)
+		{
+			StringBuilder Result = new StringBuilder();
+			Result.AppendFormat(Token.Type.ToString());
+			if(Token.Type == TokenType.Identifier || Token.Type == TokenType.Number || Token.Type == TokenType.Character || Token.Type == TokenType.String)
+			{
+				Result.AppendFormat("({0})", Token.Text);
+			}
+			return Result.ToString();
+		}
+
+		public static void RunLineNumberTest(string InputText, int ExpectedLineNumber)
+		{
+			TokenReader Reader = new TokenReader(InputText);
+			while(Reader.MoveNext())
+			{
+				if(Reader.Current.Type == TokenType.Plus)
+				{
+					break;
+				}
+			}
+			Assert.AreEqual(ExpectedLineNumber, Reader.LineNumber);
+		}
+	}
+}
diff --git a/Engine/Source/Programs/UnrealBuildTool.Tests/UnrealBuildTool.Tests.csproj b/Engine/Source/Programs/UnrealBuildTool.Tests/UnrealBuildTool.Tests.csproj
new file mode 100644
--- /dev/null
+++ b/Engine/Source/Programs/UnrealBuildTool.Tests/UnrealBuildTool.Tests.csproj
@@ -0,0 +1,20 @@
+<Project Sdk="Microsoft.NET.Sdk">
+
+  <PropertyGroup>
+    <TargetFramework>netcoreapp3.1</TargetFramework>
+
+    <IsPackable>false</IsPackable>
+  </PropertyGroup>
+
+  <ItemGroup>
+    <PackageReference Include="Microsoft.NET.Test.Sdk" Version="16.5.0" />
+    <PackageReference Include="MSTest.TestAdapter" Version="2.1.0" />
+    <PackageReference Include="MSTest.TestFramework" Version="2.1.0" />
+    <PackageReference Include="coverlet.collector" Version="1.2.0" />
+  </ItemGroup>
+
+  <ItemGroup>
+    <ProjectReference Include="..\UnrealBuildTool\UnrealBuildTool.csproj" />
+  </ItemGroup>
+
+</Project>
diff --git a/Engine/Source/Programs/UnrealBuildTool/ProjectFiles/ProjectFileGenerator.cs b/Engine/Source/Programs/UnrealBuildTool/ProjectFiles/ProjectFileGenerator.cs
--- a/Engine/Source/Programs/UnrealBuildTool/ProjectFiles/ProjectFileGenerator.cs
+++ b/Engine/Source/Programs/UnrealBuildTool/ProjectFiles/ProjectFileGenerator.cs
@@ -1739,33 +1739,30 @@ namespace UnrealBuildTool
 		/// </summary>
 		private void AddUnrealBuildToolProject(PrimaryProjectFolder ProgramsFolder)
 		{
-			List<string> ProjectDirectoryNames = new List<string>();
-			ProjectDirectoryNames.Add("UnrealBuildTool");
+			DirectoryReference ProgramsDirectory = DirectoryReference.Combine(UnrealBuildTool.EngineSourceDirectory, "Programs");
 
-			foreach (string ProjectDirectoryName in ProjectDirectoryNames)
+			FileReference ProjectFileName = FileReference.Combine(ProgramsDirectory, "UnrealBuildTool", "UnrealBuildTool.csproj");
+			if (FileReference.Exists(ProjectFileName))
 			{
-				DirectoryReference ProjectDirectory = DirectoryReference.Combine(UnrealBuildTool.EngineSourceDirectory, "Programs", ProjectDirectoryName);
-				if (DirectoryReference.Exists(ProjectDirectory))
-				{
-					string ProjectName = "UnrealBuildTool.csproj";
-
-					FileReference ProjectFileName = FileReference.Combine(ProjectDirectory, ProjectName);
+				VCSharpProjectFile UnrealBuildToolProject = new VCSharpProjectFile(ProjectFileName);
+				UnrealBuildToolProject.ShouldBuildForAllSolutionTargets = true;
 
-					if (FileReference.Exists(ProjectFileName))
-					{
-						VCSharpProjectFile UnrealBuildToolProject = new VCSharpProjectFile(ProjectFileName);
-						UnrealBuildToolProject.ShouldBuildForAllSolutionTargets = true;
+				// Store it off as we need it when generating target projects.
+				UBTProject = UnrealBuildToolProject;
 
-						// Store it off as we need it when generating target projects.
-						UBTProject = UnrealBuildToolProject;
+				// Add the project
+				AddExistingProjectFile(UnrealBuildToolProject, bNeedsAllPlatformAndConfigurations: true, bForceDevelopmentConfiguration: true);
 
-						// Add the project
-						AddExistingProjectFile(UnrealBuildToolProject, bNeedsAllPlatformAndConfigurations: true, bForceDevelopmentConfiguration: true);
+				// Put this in a solution folder
+				ProgramsFolder.ChildProjects.Add(UnrealBuildToolProject);
+			}
 
-						// Put this in a solution folder
-						ProgramsFolder.ChildProjects.Add(UnrealBuildToolProject);
-					}
-				}
+			FileReference TestsProjectFileName = FileReference.Combine(ProgramsDirectory, "UnrealBuildTool.Tests", "UnrealBuildTool.Tests.csproj");
+			if (FileReference.Exists(TestsProjectFileName))
+			{
+				VCSharpProjectFile UbtTestsProject = new VCSharpProjectFile(TestsProjectFileName);
+				AddExistingProjectFile(UbtTestsProject, bNeedsAllPlatformAndConfigurations: true, bForceDevelopmentConfiguration: true);
+				ProgramsFolder.ChildProjects.Add(UbtTestsProject);
 			}
 		}
 
diff --git a/Engine/Source/Programs/UnrealBuildToolTests/PreprocessorExpressionTests.cs b/Engine/Source/Programs/UnrealBuildToolTests/PreprocessorExpressionTests.cs
deleted file mode 100644
diff --git a/Engine/Source/Programs/UnrealBuildToolTests/PreprocessorStateTests.cs b/Engine/Source/Programs/UnrealBuildToolTests/PreprocessorStateTests.cs
deleted file mode 100644
diff --git a/Engine/Source/Programs/UnrealBuildToolTests/PreprocessorTests.cs b/Engine/Source/Programs/UnrealBuildToolTests/PreprocessorTests.cs
deleted file mode 100644
diff --git a/Engine/Source/Programs/UnrealBuildToolTests/SourceFileTests.cs b/Engine/Source/Programs/UnrealBuildToolTests/SourceFileTests.cs
deleted file mode 100644
diff --git a/Engine/Source/Programs/UnrealBuildToolTests/TokenReaderTests.cs b/Engine/Source/Programs/UnrealBuildToolTests/TokenReaderTests.cs
deleted file mode 100644
diff --git a/Engine/Source/Programs/UnrealBuildToolTests/UnrealBuildToolTests.csproj b/Engine/Source/Programs/UnrealBuildToolTests/UnrealBuildToolTests.csproj
deleted file mode 100644
