UBT: Improve matching of Mac link errors with multiple symbols.

diff --git a/Engine/Source/Programs/UnrealBuildTool.Tests/LogEventMatcherTests.cs b/Engine/Source/Programs/UnrealBuildTool.Tests/LogEventMatcherTests.cs
--- a/Engine/Source/Programs/UnrealBuildTool.Tests/LogEventMatcherTests.cs
+++ b/Engine/Source/Programs/UnrealBuildTool.Tests/LogEventMatcherTests.cs
@@ -328,6 +328,27 @@ namespace UnrealBuildToolTests
 			Assert.AreEqual("Foo::Bar", symbolProperty.Properties!["identifier"].ToString());
 		}
 
+		[TestMethod]
+		public void MacLinkEventMatcher2()
+		{
+			string[] lines =
+			{
+				@"  Undefined symbols for architecture x86_64:",
+				@"    ""_OBJC_CLASS_$_NSAppleScript"", referenced from:",
+				@"        objc-class-ref in libsupp.a(macutil.o)",
+				@"    ""_OBJC_CLASS_$_NSString"", referenced from:",
+				@"        objc-class-ref in libsupp.a(macutil.o)",
+				@"    ""_OBJC_CLASS_$_NSBundle"", referenced from:",
+				@"        objc-class-ref in libsupp.a(macutil.o)",
+				@"    ""_CGSessionCopyCurrentDictionary"", referenced from:",
+				@"        _WindowServicesAvailable in libsupp.a(macutil.o)",
+				@"  ld: symbol(s) not found for architecture x86_64"
+			};
+
+			List<LogEvent> logEvents = Parse(String.Join("\n", lines));
+			CheckEventGroup(logEvents, 0, 10, LogLevel.Error, KnownLogEvents.Linker_UndefinedSymbol);
+		}
+
 		[TestMethod]
 		public void LinuxLinkErrorMatcher()
 		{
diff --git a/Engine/Source/Programs/UnrealBuildTool/Matchers/LinkEventMatcher.cs b/Engine/Source/Programs/UnrealBuildTool/Matchers/LinkEventMatcher.cs
--- a/Engine/Source/Programs/UnrealBuildTool/Matchers/LinkEventMatcher.cs
+++ b/Engine/Source/Programs/UnrealBuildTool/Matchers/LinkEventMatcher.cs
@@ -69,14 +69,14 @@ namespace UnrealBuildTool.Matchers
 				AddSymbolMarkupForLine(builder);
 
 				string prefix = $"^(?<prefix>{match.Groups[1].Value}\\s+)";
-				if (builder.Next.TryMatch(prefix + @"""(?<symbol>[^""]+)""", out match))
+				while (builder.Next.TryMatch(prefix + @"""(?<symbol>[^""]+)""", out match))
 				{
-					prefix = $"^{match.Groups["prefix"].Value}\\s+";
+					string nextPrefix = $"^{match.Groups["prefix"].Value}\\s+";
 
 					builder.MoveNext();
 					builder.AnnotateSymbol(match.Groups["symbol"]);
 
-					while(builder.Next.TryMatch(prefix + "(?<symbol>[^ ].*) in ", out match))
+					while(builder.Next.TryMatch(nextPrefix + "(?<symbol>[^ ].*) in ", out match))
 					{
 						builder.MoveNext();
 						builder.AnnotateSymbol(match.Groups["symbol"]);
