UBT: Better matching of multi-line Clang errors. Diagnostics which have multiple "note" clauses will now be parsed correctly, and lines do not have to be indented to parse correctly.

diff --git a/Engine/Source/Programs/UnrealBuildTool.Tests/LogEventMatcherTests.cs b/Engine/Source/Programs/UnrealBuildTool.Tests/LogEventMatcherTests.cs
--- a/Engine/Source/Programs/UnrealBuildTool.Tests/LogEventMatcherTests.cs
+++ b/Engine/Source/Programs/UnrealBuildTool.Tests/LogEventMatcherTests.cs
@@ -22,6 +22,12 @@ namespace UnrealBuildToolTests
 		{
 			int _logLineIndex;
 
+			public void Reset()
+			{
+				_events.Clear();
+				_logLineIndex = 0;
+			}
+
 			public List<LogEvent> _events = new List<LogEvent>();
 
 			public IDisposable? BeginScope<TState>(TState state) => null;
@@ -175,6 +181,40 @@ namespace UnrealBuildToolTests
 			CheckEventGroup(events, 0, 5, LogLevel.Error, KnownLogEvents.Compiler);
 		}
 
+		[TestMethod]
+		public void ClangEventMatcher2()
+		{
+			// Linux Clang error
+			string[] lines =
+			{
+				@"[1/7] Compile NaniteMeshLodGroupUpdateCommandlet.cpp",
+				@"/mnt/horde/U5+RES+Inc+Min/Sync/Samples/Showcases/CitySample/Source/CitySampleEditor/Commandlets/NaniteMeshLodGroupUpdateCommandlet.cpp:33:9: error: 'ClassNames' is deprecated: Class names are now represented by path names. Please use ClassPaths. Please update your code to the new API before upgrading to the next release, otherwise your project will no longer compile. [-Werror,-Wdeprecated-declarations]",
+				@"        Filter.ClassNames.Add(FName(TEXT(""StaticMesh""))); ",
+				@"               ^",
+				@"/mnt/horde/U5+RES+Inc+Min/Sync/Engine/Source/Runtime/CoreUObject/Public/AssetRegistry/ARFilter.h:35:2: note: 'ClassNames' has been explicitly marked deprecated here",
+				@"        UE_DEPRECATED(5.1, ""Class names are now represented by path names. Please use ClassPaths."")",
+				@"        ^",
+				@"/mnt/horde/U5+RES+Inc+Min/Sync/Engine/Source/Runtime/Core/Public/Misc/CoreMiscDefines.h:232:43: note: expanded from macro 'UE_DEPRECATED'",
+				@"#define UE_DEPRECATED(Version, Message) [[deprecated(Message "" Please update your code to the new API before upgrading to the next release, otherwise your project will no longer compile."")]]",
+				@"                                          ^",
+				@"1 error generated.",
+			};
+
+			List<LogEvent> logEvents = Parse(String.Join("\n", lines));
+			CheckEventGroup(logEvents, 1, 9, LogLevel.Error, KnownLogEvents.Compiler);
+
+			LogValue fileProperty = (LogValue)logEvents[0].GetProperty("file");
+			Assert.AreEqual(@"/mnt/horde/U5+RES+Inc+Min/Sync/Samples/Showcases/CitySample/Source/CitySampleEditor/Commandlets/NaniteMeshLodGroupUpdateCommandlet.cpp", fileProperty.Text);
+			Assert.AreEqual(LogValueType.SourceFile, fileProperty.Type);
+
+			LogValue noteProperty1 = logEvents[3].GetProperty<LogValue>("file");
+			Assert.AreEqual(@"/mnt/horde/U5+RES+Inc+Min/Sync/Engine/Source/Runtime/CoreUObject/Public/AssetRegistry/ARFilter.h", noteProperty1.Text);
+			Assert.AreEqual(LogValueType.SourceFile, noteProperty1.Type);
+
+			LogValue noteProperty2 = logEvents[6].GetProperty<LogValue>("file");
+			Assert.AreEqual(@"/mnt/horde/U5+RES+Inc+Min/Sync/Engine/Source/Runtime/Core/Public/Misc/CoreMiscDefines.h", noteProperty2.Text);
+		}
+
 		[TestMethod]
 		public void IOSCompileErrorMatcher()
 		{
@@ -190,7 +230,7 @@ namespace UnrealBuildToolTests
 				@"        ^"
 			};
 
-			List<LogEvent> logEvents = Parse(String.Join("\n", lines), new DirectoryReference("[REDACTED_USER_HOME]/Build/++UE4/Sync"));
+			List<LogEvent> logEvents = Parse(String.Join("\n", lines));
 			CheckEventGroup(logEvents, 1, 1, LogLevel.Error, KnownLogEvents.Compiler);
 
 			LogEvent logEvent = logEvents[0];
@@ -228,7 +268,7 @@ namespace UnrealBuildToolTests
 					@"orbis-clang: error: linker command failed with exit code 1 (use -v to see invocation)"
 				};
 
-				List<LogEvent> logEvents = Parse(String.Join("\n", lines), new DirectoryReference("[REDACTED_USER_HOME]/Build/++UE5/Sync"));
+				List<LogEvent> logEvents = Parse(String.Join("\n", lines));
 				Assert.AreEqual(5, logEvents.Count);
 				CheckEventGroup(logEvents.Slice(0, 5), 1, 5, LogLevel.Error, KnownLogEvents.Linker_UndefinedSymbol);
 
@@ -477,11 +517,6 @@ namespace UnrealBuildToolTests
 		}
 
 		static List<LogEvent> Parse(string text)
-		{
-			return Parse(text, new DirectoryReference("C:\\Horde".Replace('\\', Path.DirectorySeparatorChar)));
-		}
-
-		static List<LogEvent> Parse(string text, DirectoryReference workspaceDir)
 		{
 			byte[] textBytes = Encoding.UTF8.GetBytes(text);
 
@@ -493,6 +528,8 @@ namespace UnrealBuildToolTests
 			{
 				parser.AddMatchersFromAssembly(typeof(UnrealBuildTool.UnrealBuildTool).Assembly);
 
+				logger.Reset();
+
 				int pos = 0;
 				while (pos < textBytes.Length)
 				{
diff --git a/Engine/Source/Programs/UnrealBuildTool/Matchers/CompileEventMatcher.cs b/Engine/Source/Programs/UnrealBuildTool/Matchers/CompileEventMatcher.cs
--- a/Engine/Source/Programs/UnrealBuildTool/Matchers/CompileEventMatcher.cs
+++ b/Engine/Source/Programs/UnrealBuildTool/Matchers/CompileEventMatcher.cs
@@ -48,6 +48,8 @@ namespace UnrealBuildTool.Matchers
 		static readonly Regex s_blankLinePattern = new Regex(@"^\s*$");
 		static readonly Regex s_errorWarningPattern = new Regex("error|warning");
 		static readonly Regex s_clangDiagnosticPattern = new Regex($"^\\s*{FilePattern}\\s*{ClangLocationPattern}:\\s*{ClangSeverity}\\s*:");
+		static readonly Regex s_clangNotePattern = new Regex($"^\\s*{FilePattern}\\s*{ClangLocationPattern}:\\s*note:");
+		static readonly Regex s_clangMarkerPattern = new Regex(@"^\s*\^$");
 
 		static readonly string[] s_invalidExtensions =
 		{
@@ -117,11 +119,18 @@ namespace UnrealBuildTool.Matchers
 					builder.TryAnnotate(match.Groups["line"], LogEventMarkup.LineNumber);
 					builder.TryAnnotate(match.Groups["column"], LogEventMarkup.ColumnNumber);
 
-					string indent = ExtractIndent(input[0]!);
-
-					Regex notePattern = new Regex($"^(?:{indent} |{indent}\\s*{FilePattern}\\s*{ClangLocationPattern}\\s*note:| *$)");
-					while (builder.Current.TryMatch(1, notePattern, out match))
+					for (; ; )
 					{
+						if (builder.Current.IsMatch(2, s_clangMarkerPattern))
+						{
+							builder.MoveNext(2);
+						}
+
+						if(!builder.Next.TryMatch(s_clangNotePattern, out match))
+						{
+							break;
+						}
+
 						builder.MoveNext();
 
 						Group fileGroup = match.Groups["file"];
