Fix UHT test suite not building/throwing an exception on startup when run as part of MSTest. Now runs ok, but tests fail.

diff --git a/Engine/Source/Programs/UnrealBuildTool.Tests/UnrealBuildTool.Tests.csproj b/Engine/Source/Programs/UnrealBuildTool.Tests/UnrealBuildTool.Tests.csproj
--- a/Engine/Source/Programs/UnrealBuildTool.Tests/UnrealBuildTool.Tests.csproj
+++ b/Engine/Source/Programs/UnrealBuildTool.Tests/UnrealBuildTool.Tests.csproj
@@ -4,6 +4,7 @@
     <TargetFramework>net6.0</TargetFramework>
 
     <IsPackable>false</IsPackable>
+    <OutputPath>..\..\..\Binaries\DotNET\UnrealBuildTool.Tests</OutputPath>
   </PropertyGroup>
 
   <ItemGroup>
diff --git a/Engine/Source/Programs/UnrealBuildTool.Tests/UnrealHeaderToolTests.cs b/Engine/Source/Programs/UnrealBuildTool.Tests/UnrealHeaderToolTests.cs
--- a/Engine/Source/Programs/UnrealBuildTool.Tests/UnrealHeaderToolTests.cs
+++ b/Engine/Source/Programs/UnrealBuildTool.Tests/UnrealHeaderToolTests.cs
@@ -1,9 +1,13 @@
 // Copyright Epic Games, Inc. All Rights Reserved.
 
+using System.IO;
+using UnrealBuildBase;
 using EpicGames.Core;
 using EpicGames.UHT.Utils;
+using Microsoft.Extensions.Logging;
 using Microsoft.VisualStudio.TestTools.UnitTesting;
 using UnrealBuildTool.Modes;
+using System.Reflection;
 
 namespace UnrealBuildToolTests
 {
@@ -13,6 +17,9 @@ namespace UnrealBuildToolTests
 		[TestMethod]
 		public void Run()
 		{
+			FileReference AssemblyFile = new FileReference(Assembly.GetExecutingAssembly().Location);
+			Unreal.LocationOverride.RootDirectory = DirectoryReference.Combine(AssemblyFile.Directory, "../../../../..");
+
 			string[] Arguments = new string[] { };
 			CommandLineArguments CommandLineArguments = new CommandLineArguments(Arguments);
 			UhtGlobalOptions Options = new UhtGlobalOptions(CommandLineArguments);
@@ -24,7 +31,8 @@ namespace UnrealBuildToolTests
 			IUhtConfig Config = new UhtConfigImpl(CommandLineArguments);
 
 			// Run the tests
-			Assert.IsTrue(UhtTestHarness.RunTests(Tables, Config, Options));
+			ILoggerFactory factory = LoggerFactory.Create(x => x.AddEpicDefault());
+			Assert.IsTrue(UhtTestHarness.RunTests(Tables, Config, Options, factory.CreateLogger<UhtTestHarness>()));
 		}
 	}
 }
