UBT: Improve matching of compile and link errors. Now handles Clang annotations for particular tokens within a line, and duplicate link errors with a named source file.

diff --git a/Engine/Source/Programs/UnrealBuildTool.Tests/LogEventMatcherTests.cs b/Engine/Source/Programs/UnrealBuildTool.Tests/LogEventMatcherTests.cs
--- a/Engine/Source/Programs/UnrealBuildTool.Tests/LogEventMatcherTests.cs
+++ b/Engine/Source/Programs/UnrealBuildTool.Tests/LogEventMatcherTests.cs
@@ -215,6 +215,32 @@ namespace UnrealBuildToolTests
 			Assert.AreEqual(@"/mnt/horde/U5+RES+Inc+Min/Sync/Engine/Source/Runtime/Core/Public/Misc/CoreMiscDefines.h", noteProperty2.Text);
 		}
 
+		[TestMethod]
+		public void ClangEventMatcher3()
+		{
+			string[] lines =
+			{
+				@"In file included from ../Intermediate/Build/Windows/UnrealGame/Development/Launch/Module.Launch.cpp:2:",
+				@"Engine/Platforms/Windows/Source/Runtime/Launch/Private/LaunchWindows.cpp(95,15): error: member access into incomplete type 'FConfigCacheIni'",
+				@"                        if (GConfig->GetBool(TEXT(""/Script/WindowsRuntimeSettings.WindowsRuntimeSettings""), TEXT(""bRouteGameUserSettingsToSaveGame""), bUseSaveForGameUserSettings, GEngineIni) && bUseSaveForGameUserSettings)",
+				@"                                   ^",
+				@"Engine/Source/Runtime/Core/Public/CoreGlobals.h(16,7): note: forward declaration of 'FConfigCacheIni'",
+				@"class FConfigCacheIni;",
+				@"      ^",
+			};
+
+			List<LogEvent> logEvents = Parse(String.Join("\n", lines));
+			CheckEventGroup(logEvents, 0, 7, LogLevel.Error, KnownLogEvents.Compiler);
+
+			LogValue fileProperty = (LogValue)logEvents[1].GetProperty("file");
+			Assert.AreEqual(@"Engine/Platforms/Windows/Source/Runtime/Launch/Private/LaunchWindows.cpp", fileProperty.Text);
+			Assert.AreEqual(LogValueType.SourceFile, fileProperty.Type);
+
+			LogValue noteProperty1 = logEvents[4].GetProperty<LogValue>("file");
+			Assert.AreEqual(@"Engine/Source/Runtime/Core/Public/CoreGlobals.h", noteProperty1.Text);
+			Assert.AreEqual(LogValueType.SourceFile, noteProperty1.Type);
+		}
+
 		[TestMethod]
 		public void IOSCompileErrorMatcher()
 		{
@@ -231,7 +257,7 @@ namespace UnrealBuildToolTests
 			};
 
 			List<LogEvent> logEvents = Parse(String.Join("\n", lines));
-			CheckEventGroup(logEvents, 1, 1, LogLevel.Error, KnownLogEvents.Compiler);
+			CheckEventGroup(logEvents, 1, 7, LogLevel.Error, KnownLogEvents.Compiler);
 
 			LogEvent logEvent = logEvents[0];
 			Assert.AreEqual("7", logEvent.GetProperty("line").ToString());
@@ -353,6 +379,23 @@ namespace UnrealBuildToolTests
 			Assert.AreEqual("USkeleton::GetBlendProfile", symbolProperty.Properties!["identifier"].ToString());
 		}
 
+		[TestMethod]
+		public void LinkEventMatcher4()
+		{
+			string[] lines =
+			{
+				@"Engine/Source/Module.DataInterface.cpp : error: L0019: symbol `IMPLEMENT_MODULE_DataInterface' multiply defined",
+				@"Engine/Source/Module.DataInterface.cpp : error:  ...first in ""D:\build\U5M+Inc\Sync\EngineTest\Intermediate\Build\PS4\EngineTest\Development\DataInterfaceGraph\Module.DataInterfaceGraph.cpp.o""",
+				@"Engine/Source/Module.DataInterface.cpp : error:  ...now in ""D:\build\U5M+Inc\Sync\EngineTest\Intermediate\Build\PS4\EngineTest\Development\DataInterface\Module.DataInterface.cpp.o"".",
+			};
+
+			List<LogEvent> logEvents = Parse(String.Join("\n", lines));
+			CheckEventGroup(logEvents, 0, 3, LogLevel.Error, KnownLogEvents.Linker_DuplicateSymbol);
+
+			LogValue symbolProperty = (LogValue)logEvents[0].GetProperty("symbol");
+			Assert.AreEqual("IMPLEMENT_MODULE_DataInterface", symbolProperty.Properties!["identifier"].ToString());
+		}
+
 		[TestMethod]
 		public void MacLinkErrorMatcher()
 		{
diff --git a/Engine/Source/Programs/UnrealBuildTool/Matchers/CompileEventMatcher.cs b/Engine/Source/Programs/UnrealBuildTool/Matchers/CompileEventMatcher.cs
--- a/Engine/Source/Programs/UnrealBuildTool/Matchers/CompileEventMatcher.cs
+++ b/Engine/Source/Programs/UnrealBuildTool/Matchers/CompileEventMatcher.cs
@@ -49,7 +49,7 @@ namespace UnrealBuildTool.Matchers
 		static readonly Regex s_errorWarningPattern = new Regex("error|warning");
 		static readonly Regex s_clangDiagnosticPattern = new Regex($"^\\s*{FilePattern}\\s*{ClangLocationPattern}:\\s*{ClangSeverity}\\s*:");
 		static readonly Regex s_clangNotePattern = new Regex($"^\\s*{FilePattern}\\s*{ClangLocationPattern}:\\s*note:");
-		static readonly Regex s_clangMarkerPattern = new Regex(@"^\s*\^$");
+		static readonly Regex s_clangMarkerPattern = new Regex(@"^(\s*)\^~*$");
 
 		static readonly string[] s_invalidExtensions =
 		{
@@ -121,10 +121,7 @@ namespace UnrealBuildTool.Matchers
 
 					for (; ; )
 					{
-						if (builder.Current.IsMatch(2, s_clangMarkerPattern))
-						{
-							builder.MoveNext(2);
-						}
+						SkipClangMarker(builder);
 
 						if(!builder.Next.TryMatch(s_clangNotePattern, out match))
 						{
@@ -206,12 +203,20 @@ namespace UnrealBuildTool.Matchers
 
 			for(; ;)
 			{
+				while (builder.Current.StartsWith(1, nextIndent))
+				{
+					builder.MoveNext();
+				}
+
+				// Clang-as-MSVC outputs warnings using its own marker syntax
+				SkipClangMarker(builder);
+
 				int offset = 1;
-				while (builder.Current.IsMatch(offset, s_blankLinePattern) || builder.Current.StartsWith(offset, nextIndent))
+				while (builder.Current.IsMatch(offset, s_blankLinePattern))
 				{
 					offset++;
 				}
-				if (!builder.Current.StartsWith(offset, indent) || !builder.Current.TryMatch(offset, s_msvcNotePattern, out match))
+				if (!builder.Current.TryMatch(offset, s_msvcNotePattern, out match))
 				{
 					break;
 				}
@@ -226,16 +231,26 @@ namespace UnrealBuildTool.Matchers
 				}
 			}
 
-			Regex pattern = new Regex($"^{indent} |: note:");
-			while (builder.Current.IsMatch(1, pattern))
-			{
-				builder.MoveNext();
-			}
-
 			outEvent = builder.ToMatch(LogEventPriority.High, level, KnownLogEvents.Compiler);
 			return true;
 		}
 
+		void SkipClangMarker(LogEventBuilder builder)
+		{
+			Match? match;
+			if (builder.Current.TryMatch(2, s_clangMarkerPattern, out match))
+			{
+				string indent = match.Groups[1].Value;
+
+				int length = 2;
+				if (builder.Current.TryGetLine(3, out string? suggestLine) && suggestLine.Length > indent.Length && suggestLine.StartsWith(indent) && !Char.IsWhiteSpace(suggestLine[indent.Length]))
+				{
+					length++;
+				}
+				builder.MoveNext(length);
+			}
+		}
+
 		static string? GetPlatformAgnosticDirectoryName(string fileName)
 		{
 			int index = fileName.LastIndexOfAny(new[] { '/', '\\' });
diff --git a/Engine/Source/Programs/UnrealBuildTool/Matchers/LinkEventMatcher.cs b/Engine/Source/Programs/UnrealBuildTool/Matchers/LinkEventMatcher.cs
--- a/Engine/Source/Programs/UnrealBuildTool/Matchers/LinkEventMatcher.cs
+++ b/Engine/Source/Programs/UnrealBuildTool/Matchers/LinkEventMatcher.cs
@@ -21,6 +21,12 @@ namespace UnrealBuildTool.Matchers
 		static readonly Regex s_ldClangPattern = new Regex(
 			@"^\s*(ld|clang):");
 
+		static readonly Regex s_ldMultiplyDefinedPattern = new Regex(
+			@"(?<severity>error): L\d+: symbol `(?<symbol>[^`]+)' multiply defined$");
+
+		static readonly Regex s_ldMultiplyDefinedInfoPattern = new Regex(
+			@"error:  \.\.\.");
+
 		static readonly Regex s_microsoftErrorPattern = new Regex(
 			@"error (?<code>LNK\d+):");
 
@@ -123,6 +129,19 @@ namespace UnrealBuildTool.Matchers
 					return builder.ToMatch(LogEventPriority.High, LogLevel.Error, KnownLogEvents.Linker);
 				}
 			}
+			if (cursor.TryMatch(s_ldMultiplyDefinedPattern, out match))
+			{
+				LogEventBuilder builder = new LogEventBuilder(cursor);
+				builder.Annotate(match.Groups["severity"], LogEventMarkup.Severity);
+				builder.AnnotateSymbol(match.Groups["symbol"]);
+
+				while (builder.Current.IsMatch(1, s_ldMultiplyDefinedInfoPattern))
+				{
+					builder.MoveNext();
+				}
+
+				return builder.ToMatch(LogEventPriority.Highest, LogLevel.Error, KnownLogEvents.Linker_DuplicateSymbol);
+			}
 
 			return null;
 		}
