UBT: Improve matching of errors from expanded macros. Now includes lines containing <scratch space> as part of the error.

diff --git a/Engine/Source/Programs/UnrealBuildTool.Tests/LogEventMatcherTests.cs b/Engine/Source/Programs/UnrealBuildTool.Tests/LogEventMatcherTests.cs
--- a/Engine/Source/Programs/UnrealBuildTool.Tests/LogEventMatcherTests.cs
+++ b/Engine/Source/Programs/UnrealBuildTool.Tests/LogEventMatcherTests.cs
@@ -241,6 +241,73 @@ namespace UnrealBuildToolTests
 			Assert.AreEqual(LogValueType.SourceFile, noteProperty1.Type);
 		}
 
+		[TestMethod]
+		public void ClangEventMatcher4()
+		{
+			string[] lines =
+			{
+				@"In file included from D:/build/++UE5/Sync/Engine/Intermediate/Build/Android/UnrealGame/Inc/NavigationSystem/UHT/AbstractNavData.gen.cpp:8:",
+				@"Engine/Source/Runtime/NavigationSystem/Public/AbstractNavData.h(45,2): error: allocating an object of abstract class type 'AAbstractNavData'",
+				@"        GENERATED_BODY()",
+				@"        ^",
+				@"Engine/Source/Runtime/CoreUObject/Public/UObject/ObjectMacros.h(687,29): note: expanded from macro 'GENERATED_BODY'",
+				@"#define GENERATED_BODY(...) BODY_MACRO_COMBINE(CURRENT_FILE_ID,_,__LINE__,_GENERATED_BODY);",
+				@"                            ^",
+				@"D:/build/++UE5/Sync/Engine/Source/Runtime/CoreUObject/Public\UObject/ObjectMacros.h(682,37): note: expanded from macro 'BODY_MACRO_COMBINE'",
+				@"#define BODY_MACRO_COMBINE(A,B,C,D) BODY_MACRO_COMBINE_INNER(A,B,C,D)",
+				@"                                    ^",
+				@"D:/build/++UE5/Sync/Engine/Source/Runtime/CoreUObject/Public\UObject/ObjectMacros.h(681,43): note: expanded from macro 'BODY_MACRO_COMBINE_INNER'",
+				@"#define BODY_MACRO_COMBINE_INNER(A,B,C,D) A##B##C##D",
+				@"                                          ^",
+				@"<scratch space>(114,1): note: expanded from here",
+				@"FID_Engine_Source_Runtime_NavigationSystem_Public_AbstractNavData_h_45_GENERATED_BODY",
+				@"^",
+				@"D:/build/++UE5/Sync/Engine/Intermediate/Build/Android/UnrealGame/Inc/NavigationSystem/UHT\AbstractNavData.generated.h(82,2): note: expanded from macro 'FID_Engine_Source_Runtime_NavigationSystem_Public_AbstractNavData_h_45_GENERATED_BODY'",
+				@"        FID_Engine_Source_Runtime_NavigationSystem_Public_AbstractNavData_h_45_ENHANCED_CONSTRUCTORS \",
+				@"        ^",
+				@"D:/build/++UE5/Sync/Engine/Intermediate/Build/Android/UnrealGame/Inc/NavigationSystem/UHT\AbstractNavData.generated.h(59,53): note: expanded from macro 'FID_Engine_Source_Runtime_NavigationSystem_Public_AbstractNavData_h_45_ENHANCED_CONSTRUCTORS'",
+				@"        DEFINE_DEFAULT_OBJECT_INITIALIZER_CONSTRUCTOR_CALL(AAbstractNavData)",
+				@"                                                           ^",
+				@"D:/build/++UE5/Sync/Engine/Source/Runtime/NavigationSystem/Public/NavigationData.h(836,15): note: unimplemented pure virtual method 'FindOverlappingEdges' in 'AAbstractNavData'",
+				@"        virtual bool FindOverlappingEdges(const FNavLocation& StartLocation, TConstArrayView<FVector> ConvexPolygon, TArray<FVector>& OutEdges, FSharedConstNavQueryFilter Filter = NULL, const UObject* Querier = NULL) const PURE_VIRTUAL(ANavigationData::FindOverlappingEdges, return false;);",
+				@"                     ^",
+				@"1 error generated."
+			};
+
+			List<LogEvent> logEvents = Parse(String.Join("\n", lines));
+			CheckEventGroup(logEvents, 0, 25, LogLevel.Error, KnownLogEvents.Compiler);
+
+			LogValue fileProperty = (LogValue)logEvents[1].GetProperty("file");
+			Assert.AreEqual(@"Engine/Source/Runtime/NavigationSystem/Public/AbstractNavData.h", fileProperty.Text);
+			Assert.AreEqual(LogValueType.SourceFile, fileProperty.Type);
+
+			LogValue noteProperty1 = logEvents[4].GetProperty<LogValue>("file");
+			Assert.AreEqual(@"Engine/Source/Runtime/CoreUObject/Public/UObject/ObjectMacros.h", noteProperty1.Text);
+			Assert.AreEqual(LogValueType.SourceFile, noteProperty1.Type);
+
+			LogValue noteProperty2 = logEvents[7].GetProperty<LogValue>("file");
+			Assert.AreEqual(@"D:/build/++UE5/Sync/Engine/Source/Runtime/CoreUObject/Public\UObject/ObjectMacros.h", noteProperty2.Text);
+			Assert.AreEqual(LogValueType.SourceFile, noteProperty2.Type);
+
+			LogValue noteProperty3 = logEvents[10].GetProperty<LogValue>("file");
+			Assert.AreEqual(@"D:/build/++UE5/Sync/Engine/Source/Runtime/CoreUObject/Public\UObject/ObjectMacros.h", noteProperty3.Text);
+			Assert.AreEqual(LogValueType.SourceFile, noteProperty3.Type);
+
+			Assert.IsFalse(logEvents[13].TryGetProperty<LogValue>("file", out _));
+
+			LogValue noteProperty5 = logEvents[16].GetProperty<LogValue>("file");
+			Assert.AreEqual(@"D:/build/++UE5/Sync/Engine/Intermediate/Build/Android/UnrealGame/Inc/NavigationSystem/UHT\AbstractNavData.generated.h", noteProperty5.Text);
+			Assert.AreEqual(LogValueType.SourceFile, noteProperty5.Type);
+
+			LogValue noteProperty6 = logEvents[19].GetProperty<LogValue>("file");
+			Assert.AreEqual(@"D:/build/++UE5/Sync/Engine/Intermediate/Build/Android/UnrealGame/Inc/NavigationSystem/UHT\AbstractNavData.generated.h", noteProperty6.Text);
+			Assert.AreEqual(LogValueType.SourceFile, noteProperty6.Type);
+
+			LogValue noteProperty7 = logEvents[22].GetProperty<LogValue>("file");
+			Assert.AreEqual(@"D:/build/++UE5/Sync/Engine/Source/Runtime/NavigationSystem/Public/NavigationData.h", noteProperty7.Text);
+			Assert.AreEqual(LogValueType.SourceFile, noteProperty7.Type);
+		}
+
 		[TestMethod]
 		public void IOSCompileErrorMatcher()
 		{
diff --git a/Engine/Source/Programs/UnrealBuildTool/Matchers/CompileEventMatcher.cs b/Engine/Source/Programs/UnrealBuildTool/Matchers/CompileEventMatcher.cs
--- a/Engine/Source/Programs/UnrealBuildTool/Matchers/CompileEventMatcher.cs
+++ b/Engine/Source/Programs/UnrealBuildTool/Matchers/CompileEventMatcher.cs
@@ -15,7 +15,7 @@ namespace UnrealBuildTool.Matchers
 	class CompileEventMatcher : ILogEventMatcher
 	{
 		const string FilePattern =
-			@"(?<file>" +
+			@"(?:(?<file>" +
 				// optional drive letter
 				@"(?:[a-zA-Z]:)?" +
 				// any non-colon character
@@ -24,7 +24,8 @@ namespace UnrealBuildTool.Matchers
 				@"[^:<>*?""]+" +
 				// valid source file extension
 				@"\.(?:(?i)(?:h|c|cc|cpp|inc|inl|cs|targets))" +
-			@")";
+				// or the string "<scratch space>"
+			@")|<scratch space>)";
 
 		const string VisualCppLocationPattern =
 			@"\(" +
@@ -243,7 +244,7 @@ namespace UnrealBuildTool.Matchers
 				string indent = match.Groups[1].Value;
 
 				int length = 2;
-				if (builder.Current.TryGetLine(3, out string? suggestLine) && suggestLine.Length > indent.Length && suggestLine.StartsWith(indent) && !Char.IsWhiteSpace(suggestLine[indent.Length]))
+				if (indent.Length > 0 && builder.Current.TryGetLine(3, out string? suggestLine) && suggestLine.Length > indent.Length && suggestLine.StartsWith(indent) && !Char.IsWhiteSpace(suggestLine[indent.Length]))
 				{
 					length++;
 				}
