Add a test to make sure we match output generated by "#pragma warning" correctly.

diff --git a/Engine/Source/Programs/Horde/Horde.Agent.Tests/LogParserTests.cs b/Engine/Source/Programs/Horde/Horde.Agent.Tests/LogParserTests.cs
--- a/Engine/Source/Programs/Horde/Horde.Agent.Tests/LogParserTests.cs
+++ b/Engine/Source/Programs/Horde/Horde.Agent.Tests/LogParserTests.cs
@@ -153,6 +153,21 @@ namespace Horde.Agent.Tests
 			CheckEventGroup(logEvents, 1, 5, LogLevel.Error, KnownLogEvents.ExitCode);
 		}
 
+		[TestMethod]
+		public void ExitCodeEventMatcher2()
+		{
+			string[] lines =
+			{
+				@"1 error generated.",
+				@"",
+				@"Error executing d:\build\AutoSDK\Sync\HostWin64\Android\-24\ndk\21.4.7075529\toolchains\llvm\prebuilt\windows-x86_64\bin\clang++.exe (tool returned code: 1)",
+				@"AnimNode_BlendListByInt.gen.cpp [arm64] (0:22.43 at +2:47)",
+			};
+
+			List<LogEvent> logEvents = Parse(lines);
+			CheckEventGroup(logEvents, 1, 5, LogLevel.Error, KnownLogEvents.ExitCode);
+		}
+
 		[TestMethod]
 		public void CrashEventMatcher()
 		{
diff --git a/Engine/Source/Programs/UnrealBuildTool.Tests/LogEventMatcherTests.cs b/Engine/Source/Programs/UnrealBuildTool.Tests/LogEventMatcherTests.cs
--- a/Engine/Source/Programs/UnrealBuildTool.Tests/LogEventMatcherTests.cs
+++ b/Engine/Source/Programs/UnrealBuildTool.Tests/LogEventMatcherTests.cs
@@ -308,6 +308,28 @@ namespace UnrealBuildToolTests
 			Assert.AreEqual(LogValueType.SourceFile, noteProperty7.Type);
 		}
 
+		[TestMethod]
+		public void ClangEventMatcher5()
+		{
+			string[] lines =
+			{
+				@"Module.StateTreeModule.cpp (0:25.07 at +6:14)",
+				@"In file included from ..\Plugins\Runtime\StateTree\Intermediate\Build\Win64\UnrealEditor\Debug\StateTreeModule\Module.StateTreeModule.cpp:10:",
+				@"In file included from .\../Plugins/Runtime/StateTree/Source/StateTreeModule/Private/StateTree.cpp:8:",
+				@"Engine/Source/Runtime/AssetRegistry/Public/AssetData.h(6,9): warning: Runtime\AssetRegistry\Public\AssetData.h(6): warning: #include AssetRegistry/AssetData.h instead of AssetData.h [-W#pragma-messages]",
+				@"#pragma message (__FILE__""(6): warning: #include AssetRegistry/AssetData.h instead of AssetData.h"")",
+				@"        ^",
+				@"1 warning generated."
+			};
+
+			List<LogEvent> logEvents = Parse(String.Join("\n", lines));
+			CheckEventGroup(logEvents, 1, 5, LogLevel.Warning, KnownLogEvents.Compiler);
+
+			LogValue fileProperty = (LogValue)logEvents[2].GetProperty("file");
+			Assert.AreEqual(@"Engine/Source/Runtime/AssetRegistry/Public/AssetData.h", fileProperty.Text);
+			Assert.AreEqual(LogValueType.SourceFile, fileProperty.Type);
+		}
+
 		[TestMethod]
 		public void IOSCompileErrorMatcher()
 		{
