Include the name of the file being compiled within the error summary.

diff --git a/Engine/Source/Programs/UnrealBuildTool.Tests/LogEventMatcherTests.cs b/Engine/Source/Programs/UnrealBuildTool.Tests/LogEventMatcherTests.cs
--- a/Engine/Source/Programs/UnrealBuildTool.Tests/LogEventMatcherTests.cs
+++ b/Engine/Source/Programs/UnrealBuildTool.Tests/LogEventMatcherTests.cs
@@ -201,17 +201,17 @@ namespace UnrealBuildToolTests
 			};
 
 			List<LogEvent> logEvents = Parse(String.Join("\n", lines));
-			CheckEventGroup(logEvents, 1, 9, LogLevel.Error, KnownLogEvents.Compiler);
+			CheckEventGroup(logEvents, 0, 10, LogLevel.Error, KnownLogEvents.Compiler);
 
-			LogValue fileProperty = (LogValue)logEvents[0].GetProperty("file");
+			LogValue fileProperty = (LogValue)logEvents[1].GetProperty("file");
 			Assert.AreEqual(@"/mnt/horde/U5+RES+Inc+Min/Sync/Samples/Showcases/CitySample/Source/CitySampleEditor/Commandlets/NaniteMeshLodGroupUpdateCommandlet.cpp", fileProperty.Text);
 			Assert.AreEqual(LogValueType.SourceFile, fileProperty.Type);
 
-			LogValue noteProperty1 = logEvents[3].GetProperty<LogValue>("file");
+			LogValue noteProperty1 = logEvents[4].GetProperty<LogValue>("file");
 			Assert.AreEqual(@"/mnt/horde/U5+RES+Inc+Min/Sync/Engine/Source/Runtime/CoreUObject/Public/AssetRegistry/ARFilter.h", noteProperty1.Text);
 			Assert.AreEqual(LogValueType.SourceFile, noteProperty1.Type);
 
-			LogValue noteProperty2 = logEvents[6].GetProperty<LogValue>("file");
+			LogValue noteProperty2 = logEvents[7].GetProperty<LogValue>("file");
 			Assert.AreEqual(@"/mnt/horde/U5+RES+Inc+Min/Sync/Engine/Source/Runtime/Core/Public/Misc/CoreMiscDefines.h", noteProperty2.Text);
 		}
 
@@ -346,9 +346,9 @@ namespace UnrealBuildToolTests
 			};
 
 			List<LogEvent> logEvents = Parse(String.Join("\n", lines));
-			CheckEventGroup(logEvents, 1, 7, LogLevel.Error, KnownLogEvents.Compiler);
+			CheckEventGroup(logEvents, 0, 8, LogLevel.Error, KnownLogEvents.Compiler);
 
-			LogEvent logEvent = logEvents[0];
+			LogEvent logEvent = logEvents[1];
 			Assert.AreEqual("7", logEvent.GetProperty("line").ToString());
 			Assert.AreEqual("48", logEvent.GetProperty("column").ToString());
 
diff --git a/Engine/Source/Programs/UnrealBuildTool/Matchers/CompileEventMatcher.cs b/Engine/Source/Programs/UnrealBuildTool/Matchers/CompileEventMatcher.cs
--- a/Engine/Source/Programs/UnrealBuildTool/Matchers/CompileEventMatcher.cs
+++ b/Engine/Source/Programs/UnrealBuildTool/Matchers/CompileEventMatcher.cs
@@ -45,7 +45,7 @@ namespace UnrealBuildTool.Matchers
 		const string ClangSeverity =
 			@"(?<severity>error|warning)";
 
-		static readonly Regex s_preludePattern = new Regex("^\\s*(?:In (member )?function|In file included from)");
+		static readonly Regex s_preludePattern = new Regex(@"^\s*(?:In (member )?function|In file included from|\[[\d/]+\] Compile |[^/\ :]+\\.cpp\s*(?:\([^\)]*\))?$)");
 		static readonly Regex s_blankLinePattern = new Regex(@"^\s*$");
 		static readonly Regex s_errorWarningPattern = new Regex("error|warning");
 		static readonly Regex s_clangDiagnosticPattern = new Regex($"^\\s*{FilePattern}\\s*{ClangLocationPattern}:\\s*{ClangSeverity}\\s*:");
