UBT: Fix over-eager matching of multiple lines in prelude to error message.

diff --git a/Engine/Source/Programs/UnrealBuildTool.Tests/LogEventMatcherTests.cs b/Engine/Source/Programs/UnrealBuildTool.Tests/LogEventMatcherTests.cs
--- a/Engine/Source/Programs/UnrealBuildTool.Tests/LogEventMatcherTests.cs
+++ b/Engine/Source/Programs/UnrealBuildTool.Tests/LogEventMatcherTests.cs
@@ -188,6 +188,7 @@ namespace UnrealBuildToolTests
 			string[] lines =
 			{
 				@"[1/7] Compile NaniteMeshLodGroupUpdateCommandlet.cpp",
+				@"[2/7] Compile NaniteMeshLodGroupUpdateCommandlet.cpp",
 				@"/mnt/horde/U5+RES+Inc+Min/Sync/Samples/Showcases/CitySample/Source/CitySampleEditor/Commandlets/NaniteMeshLodGroupUpdateCommandlet.cpp:33:9: error: 'ClassNames' is deprecated: Class names are now represented by path names. Please use ClassPaths. Please update your code to the new API before upgrading to the next release, otherwise your project will no longer compile. [-Werror,-Wdeprecated-declarations]",
 				@"        Filter.ClassNames.Add(FName(TEXT(""StaticMesh""))); ",
 				@"               ^",
@@ -201,7 +202,7 @@ namespace UnrealBuildToolTests
 			};
 
 			List<LogEvent> logEvents = Parse(String.Join("\n", lines));
-			CheckEventGroup(logEvents, 0, 10, LogLevel.Error, KnownLogEvents.Compiler);
+			CheckEventGroup(logEvents, 1, 10, LogLevel.Error, KnownLogEvents.Compiler);
 
 			LogValue fileProperty = (LogValue)logEvents[1].GetProperty("file");
 			Assert.AreEqual(@"/mnt/horde/U5+RES+Inc+Min/Sync/Samples/Showcases/CitySample/Source/CitySampleEditor/Commandlets/NaniteMeshLodGroupUpdateCommandlet.cpp", fileProperty.Text);
diff --git a/Engine/Source/Programs/UnrealBuildTool/Matchers/CompileEventMatcher.cs b/Engine/Source/Programs/UnrealBuildTool/Matchers/CompileEventMatcher.cs
--- a/Engine/Source/Programs/UnrealBuildTool/Matchers/CompileEventMatcher.cs
+++ b/Engine/Source/Programs/UnrealBuildTool/Matchers/CompileEventMatcher.cs
@@ -45,7 +45,8 @@ namespace UnrealBuildTool.Matchers
 		const string ClangSeverity =
 			@"(?<severity>error|warning)";
 
-		static readonly Regex s_preludePattern = new Regex(@"^\s*(?:In (member )?function|In file included from|\[[\d/]+\] Compile |[^/\ :]+\\.cpp\s*(?:\([^\)]*\))?$)");
+		static readonly Regex s_baseFilePattern = new Regex(@"^\s*(?:\[[\d/]+\] Compile |[^/\ :]+\\.cpp\s*(?:\([^\)]*\))?$)");
+		static readonly Regex s_preludePattern = new Regex(@"^\s*(?:In (member )?function|In file included from)");
 		static readonly Regex s_blankLinePattern = new Regex(@"^\s*$");
 		static readonly Regex s_errorWarningPattern = new Regex("error|warning");
 		static readonly Regex s_clangDiagnosticPattern = new Regex($"^\\s*{FilePattern}\\s*{ClangLocationPattern}:\\s*{ClangSeverity}\\s*:");
@@ -66,6 +67,10 @@ namespace UnrealBuildTool.Matchers
 		{
 			// Match the prelude to any error
 			int maxOffset = 0;
+			if (input.IsMatch(maxOffset, s_baseFilePattern))
+			{
+				maxOffset++;
+			}
 			while (input.IsMatch(maxOffset, s_preludePattern))
 			{
 				maxOffset++;
