Manually mask out IOS link errors from event matcher.

diff --git a/Engine/Source/Programs/UnrealBuildTool.Tests/LogEventMatcherTests.cs b/Engine/Source/Programs/UnrealBuildTool.Tests/LogEventMatcherTests.cs
--- a/Engine/Source/Programs/UnrealBuildTool.Tests/LogEventMatcherTests.cs
+++ b/Engine/Source/Programs/UnrealBuildTool.Tests/LogEventMatcherTests.cs
@@ -620,6 +620,20 @@ namespace UnrealBuildToolTests
 			Assert.AreEqual("Foo::Bar", symbolProperty.Properties!["identifier"].ToString());
 		}
 
+		[TestMethod]
+		public void IOSIgnoredLinkEvents()
+		{
+			string[] lines =
+			{
+				@"ld: warning: Linker asked to preserve internal global: 'inflateEnd'",
+				@"ld: warning: Linker asked to preserve internal global: 'inflateReset'",
+			};
+
+			List<LogEvent> logEvents = Parse(String.Join("\n", lines));
+			CheckEventGroup(logEvents.Slice(0, 1), 0, 1, LogLevel.Information, KnownLogEvents.Linker);
+			CheckEventGroup(logEvents.Slice(1, 1), 1, 1, LogLevel.Information, KnownLogEvents.Linker);
+		}
+
 		[TestMethod]
 		public void XoreaxErrorMatcher()
 		{
diff --git a/Engine/Source/Programs/UnrealBuildTool/Matchers/LinkEventMatcher.cs b/Engine/Source/Programs/UnrealBuildTool/Matchers/LinkEventMatcher.cs
--- a/Engine/Source/Programs/UnrealBuildTool/Matchers/LinkEventMatcher.cs
+++ b/Engine/Source/Programs/UnrealBuildTool/Matchers/LinkEventMatcher.cs
@@ -30,9 +30,18 @@ namespace UnrealBuildTool.Matchers
 		static readonly Regex s_microsoftErrorPattern = new Regex(
 			@"error (?<code>LNK\d+):");
 
+		static readonly Regex s_ignoredLinkWarningPattern = new Regex(
+			@"^\s*ld: warning: Linker asked to preserve internal global:");
+
 		/// <inheritdoc/>
 		public LogEventMatch? Match(ILogCursor cursor)
 		{
+			if (cursor.IsMatch(s_ignoredLinkWarningPattern))
+			{
+				LogEventBuilder builder = new LogEventBuilder(cursor);
+				return builder.ToMatch(LogEventPriority.High, LogLevel.Information, KnownLogEvents.Linker);
+			}
+
 			int lineCount = 0;
 			bool isError = false;
 			bool isWarning = false;
