Tag files referenced via "In file included from..." clauses in Clang errors.

diff --git a/Engine/Source/Programs/UnrealBuildTool.Tests/LogEventMatcherTests.cs b/Engine/Source/Programs/UnrealBuildTool.Tests/LogEventMatcherTests.cs
--- a/Engine/Source/Programs/UnrealBuildTool.Tests/LogEventMatcherTests.cs
+++ b/Engine/Source/Programs/UnrealBuildTool.Tests/LogEventMatcherTests.cs
@@ -331,6 +331,54 @@ namespace UnrealBuildToolTests
 			Assert.AreEqual(LogValueType.SourceFile, fileProperty.Type);
 		}
 
+		[TestMethod]
+		public void ClangEventMatcher6()
+		{
+			string[] lines =
+			{
+				@"[19884/26184] Compile KismetCastingUtils.cpp",
+				@"In file included from Editor/KismetCompiler/Private/KismetCastingUtils.cpp:3:",
+				@"[REDACTED_USER_HOME]/Build/++UE5/Sync/Engine/Source/Editor/KismetCompiler/Public/KismetCastingUtils.h:33:3: error: unknown type name 'UEdGraphNode'; did you mean 'UEdGraphPin'?",
+				@"                UEdGraphNode* TargetNode = nullptr;",
+				@"                ^~~~~~~~~~~~",
+				@"                UEdGraphPin",
+				@"[REDACTED_USER_HOME]/Build/++UE5/Sync/Engine/Source/Editor/KismetCompiler/Public/KismetCastingUtils.h:10:7: note: 'UEdGraphPin' declared here",
+				@"class UEdGraphPin;",
+				@"      ^",
+				@"[REDACTED_USER_HOME]/Build/++UE5/Sync/Engine/Source/Editor/KismetCompiler/Public/KismetCastingUtils.h:47:110: error: unknown type name 'UEdGraphNode'; did you mean 'UEdGraphPin'?",
+				@"        KISMETCOMPILER_API FBPTerminal* MakeImplicitCastTerminal(FKismetFunctionContext& Context, UEdGraphPin* Net, UEdGraphNode* SourceNode = nullptr);",
+				@"                                                                                                                    ^~~~~~~~~~~~",
+				@"                                                                                                                    UEdGraphPin",
+				@"[REDACTED_USER_HOME]/Build/++UE5/Sync/Engine/Source/Editor/KismetCompiler/Public/KismetCastingUtils.h:10:7: note: 'UEdGraphPin' declared here",
+				@"class UEdGraphPin;",
+				@"      ^",
+				@"[REDACTED_USER_HOME]/Build/++UE5/Sync/Engine/Source/Editor/KismetCompiler/Private/KismetCastingUtils.cpp:34:88: error: cannot initialize a member subobject of type 'UEdGraphPin *' with an lvalue of type 'UEdGraphNode *'",
+				@"                Context.ImplicitCastMap.Add(DestinationPin, FImplicitCastParams{Conversion, NewTerm, OwningNode});",
+				@"                                                                                                     ^~~~~~~~~~",
+				@"[REDACTED_USER_HOME]/Build/++UE5/Sync/Engine/Source/Editor/KismetCompiler/Private/KismetCastingUtils.cpp:121:78: error: cannot initialize a parameter of type 'UEdGraphNode *' with an lvalue of type 'UEdGraphPin *const'",
+				@"        FBlueprintCompiledStatement& CastStatement = Context.AppendStatementForNode(CastParams.TargetNode);",
+				@"                                                                                    ^~~~~~~~~~~~~~~~~~~~~",
+				@"[REDACTED_USER_HOME]/Build/++UE5/Sync/Engine/Source/Editor/KismetCompiler/Public/KismetCompiledFunctionContext.h:265:68: note: passing argument to parameter 'Node' here",
+				@"        FBlueprintCompiledStatement& AppendStatementForNode(UEdGraphNode* Node)",
+				@"                                                                          ^",
+				@"[REDACTED_USER_HOME]/Build/++UE5/Sync/Engine/Source/Editor/KismetCompiler/Private/KismetCastingUtils.cpp:121:31: error: non-const lvalue reference to type 'FBlueprintCompiledStatement' cannot bind to a temporary of type 'FBlueprintCompiledStatement'",
+				@"        FBlueprintCompiledStatement& CastStatement = Context.AppendStatementForNode(CastParams.TargetNode);",
+				@"                                     ^               ~~~~~~~~~~~~~~~~~~~~~~~~~~~~~~~~~~~~~~~~~~~~~~~~~~~~~",
+				@"5 errors generated."
+			};
+
+			List<LogEvent> logEvents = Parse(String.Join("\n", lines));
+			CheckEventGroup(logEvents.Slice(0, 9), 0, 9, LogLevel.Error, KnownLogEvents.Compiler);
+			CheckEventGroup(logEvents.Slice(9, 7), 9, 7, LogLevel.Error, KnownLogEvents.Compiler);
+			CheckEventGroup(logEvents.Slice(16, 3), 16, 3, LogLevel.Error, KnownLogEvents.Compiler);
+			CheckEventGroup(logEvents.Slice(19, 6), 19, 6, LogLevel.Error, KnownLogEvents.Compiler);
+			CheckEventGroup(logEvents.Slice(25, 3), 25, 3, LogLevel.Error, KnownLogEvents.Compiler);
+
+			LogValue fileProperty = (LogValue)logEvents[1].GetProperty("file");
+			Assert.AreEqual(@"Editor/KismetCompiler/Private/KismetCastingUtils.cpp", fileProperty.Text);
+			Assert.AreEqual(LogValueType.SourceFile, fileProperty.Type);
+		}
+
 		[TestMethod]
 		public void IOSCompileErrorMatcher()
 		{
diff --git a/Engine/Source/Programs/UnrealBuildTool/Matchers/CompileEventMatcher.cs b/Engine/Source/Programs/UnrealBuildTool/Matchers/CompileEventMatcher.cs
--- a/Engine/Source/Programs/UnrealBuildTool/Matchers/CompileEventMatcher.cs
+++ b/Engine/Source/Programs/UnrealBuildTool/Matchers/CompileEventMatcher.cs
@@ -47,11 +47,12 @@ namespace UnrealBuildTool.Matchers
 
 		static readonly Regex s_baseFilePattern = new Regex(@"^\s*(?:\[[\d/]+\] Compile |[^/\ :]+\\.cpp\s*(?:\([^\)]*\))?$)");
 		static readonly Regex s_preludePattern = new Regex(@"^\s*(?:In (member )?function|In file included from)");
+		static readonly Regex s_preludeFilePattern = new Regex($"^\\s*In file included from {FilePattern}:");
 		static readonly Regex s_blankLinePattern = new Regex(@"^\s*$");
 		static readonly Regex s_errorWarningPattern = new Regex("error|warning");
 		static readonly Regex s_clangDiagnosticPattern = new Regex($"^\\s*{FilePattern}\\s*{ClangLocationPattern}:\\s*{ClangSeverity}\\s*:");
 		static readonly Regex s_clangNotePattern = new Regex($"^\\s*{FilePattern}\\s*{ClangLocationPattern}:\\s*note:");
-		static readonly Regex s_clangMarkerPattern = new Regex(@"^(\s*)\^~*$");
+		static readonly Regex s_clangMarkerPattern = new Regex(@"^(\s*)\^\s*~*$");
 
 		static readonly string[] s_invalidExtensions =
 		{
@@ -80,7 +81,17 @@ namespace UnrealBuildTool.Matchers
 			// produce many false positives, making them very slow to execute.
 			if (input.IsMatch(maxOffset, s_errorWarningPattern))
 			{
-				LogEventBuilder builder = new LogEventBuilder(input, maxOffset + 1);
+				// Tag any files in the prelude with their source files
+				LogEventBuilder builder = new LogEventBuilder(input);
+				for (int idx = 0; idx < maxOffset; idx++)
+				{
+					Match? fileMatch;
+					if (builder.Current.TryMatch(s_preludeFilePattern, out fileMatch))
+					{
+						builder.AnnotateSourceFile(fileMatch.Groups["file"], DefaultSourceFileBaseDir);
+					}
+					builder.MoveNext();
+				}
 
 				// Try to match a Visual C++ diagnostic
 				LogEventMatch? eventMatch;
@@ -120,7 +131,7 @@ namespace UnrealBuildTool.Matchers
 				{
 					LogLevel level = GetLogLevelFromSeverity(match);
 
-					builder.AnnotateSourceFile(match.Groups["file"], "Engine/Source");
+					builder.AnnotateSourceFile(match.Groups["file"], DefaultSourceFileBaseDir);
 					builder.Annotate(match.Groups["severity"], LogEventMarkup.Severity);
 					builder.TryAnnotate(match.Groups["line"], LogEventMarkup.LineNumber);
 					builder.TryAnnotate(match.Groups["column"], LogEventMarkup.ColumnNumber);
