UBT: Add an event matcher for Clang fatal errors.

diff --git a/Engine/Source/Programs/UnrealBuildTool.Tests/LogEventMatcherTests.cs b/Engine/Source/Programs/UnrealBuildTool.Tests/LogEventMatcherTests.cs
--- a/Engine/Source/Programs/UnrealBuildTool.Tests/LogEventMatcherTests.cs
+++ b/Engine/Source/Programs/UnrealBuildTool.Tests/LogEventMatcherTests.cs
@@ -379,6 +379,23 @@ namespace UnrealBuildToolTests
 			Assert.AreEqual(LogValueType.SourceFile, fileProperty.Type);
 		}
 
+		[TestMethod]
+		public void ClangEventMatcher7()
+		{
+			string[] lines =
+			{
+				@"In file included from ../Plugins/FX/Niagara/Intermediate/Build/Linux/B4D820EA/UnrealEditor/Development/NiagaraEditor/Module.NiagaraEditor.1_of_12.cpp:25:",
+				@"Engine/Plugins/FX/Niagara/Source/NiagaraEditor/Private/Customizations/NiagaraTypeCustomizations.cpp:31:10: fatal error: 'Framework/Multibox/MultiBoxBuilder.h' file not found"
+			};
+
+			List<LogEvent> logEvents = Parse(String.Join("\n", lines));
+			CheckEventGroup(logEvents, 0, 2, LogLevel.Error, KnownLogEvents.Compiler);
+
+			LogValue fileProperty = (LogValue)logEvents[1].GetProperty("file");
+			Assert.AreEqual(@"Engine/Plugins/FX/Niagara/Source/NiagaraEditor/Private/Customizations/NiagaraTypeCustomizations.cpp", fileProperty.Text);
+			Assert.AreEqual(LogValueType.SourceFile, fileProperty.Type);
+		}
+
 		[TestMethod]
 		public void IOSCompileErrorMatcher()
 		{
diff --git a/Engine/Source/Programs/UnrealBuildTool/Matchers/CompileEventMatcher.cs b/Engine/Source/Programs/UnrealBuildTool/Matchers/CompileEventMatcher.cs
--- a/Engine/Source/Programs/UnrealBuildTool/Matchers/CompileEventMatcher.cs
+++ b/Engine/Source/Programs/UnrealBuildTool/Matchers/CompileEventMatcher.cs
@@ -43,7 +43,7 @@ namespace UnrealBuildTool.Matchers
 			@"";
 
 		const string ClangSeverity =
-			@"(?<severity>error|warning)";
+			@"(?<severity>error|warning|fatal error)";
 
 		static readonly Regex s_baseFilePattern = new Regex(@"^\s*(?:\[[\d/]+\] Compile |[^/\ :]+\\.cpp\s*(?:\([^\)]*\))?$)");
 		static readonly Regex s_preludePattern = new Regex(@"^\s*(?:In (member )?function|In file included from)");
