UBT: Add matcher for errors/warnings from UHT.

diff --git a/Engine/Source/Programs/UnrealBuildTool.Tests/LogEventMatcherTests.cs b/Engine/Source/Programs/UnrealBuildTool.Tests/LogEventMatcherTests.cs
--- a/Engine/Source/Programs/UnrealBuildTool.Tests/LogEventMatcherTests.cs
+++ b/Engine/Source/Programs/UnrealBuildTool.Tests/LogEventMatcherTests.cs
@@ -723,6 +723,26 @@ namespace UnrealBuildToolTests
 			CheckEventGroup(logEvents.Slice(7, 1), 8, 1, LogLevel.Error, KnownLogEvents.Systemic_Xge_BuildFailed);
 		}
 
+		[TestMethod]
+		public void UhtErrorMatcher()
+		{
+			string[] lines =
+			{
+				@"Parsing headers for UnrealEditor",
+				@"  Running UnrealHeaderTool UnrealEditor ""/mnt/horde/FN+NC+PF/Sync/Engine/Intermediate/Build/Linux/B4D820EA/UnrealEditor/Development/UnrealEditor.uhtmanifest"" -LogCmds=""loginit warning, logexit warning, logdatabase error"" -Unattended -WarningsAsErrors -abslog=""/mnt/horde/FN+NC+PF/Sync/Engine/Programs/AutomationTool/Saved/Logs/UHT-UnrealEditor-Linux-Development.txt""",
+				@"LogStreaming: Display: NotifyRegistrationEvent: Replay 112 entries",
+				@"/mnt/horde/FN+NC+PF/Sync/Engine/Plugins/Experimental/DataInterfaceGraph/Source/DataInterfaceGraphUncookedOnly/Internal/DataInterfaceGraph_EditorData.h(91): Error: Do not specify struct property containing editor only properties inside an optional class.",
+			};
+
+			List<LogEvent> logEvents = Parse(lines);
+			Assert.AreEqual(1, logEvents.Count);
+			CheckEventGroup(logEvents, 3, 1, LogLevel.Error, KnownLogEvents.Compiler);
+
+			LogValue fileProperty = (LogValue)logEvents[0].GetProperty("file");
+			Assert.AreEqual(@"/mnt/horde/FN+NC+PF/Sync/Engine/Plugins/Experimental/DataInterfaceGraph/Source/DataInterfaceGraphUncookedOnly/Internal/DataInterfaceGraph_EditorData.h", fileProperty.Text);
+			Assert.AreEqual(LogValueType.SourceFile, fileProperty.Type);
+		}
+
 		static List<LogEvent> Parse(IEnumerable<string> lines)
 		{
 			return Parse(String.Join("\n", lines));
diff --git a/Engine/Source/Programs/UnrealBuildTool/Matchers/CompileEventMatcher.cs b/Engine/Source/Programs/UnrealBuildTool/Matchers/CompileEventMatcher.cs
--- a/Engine/Source/Programs/UnrealBuildTool/Matchers/CompileEventMatcher.cs
+++ b/Engine/Source/Programs/UnrealBuildTool/Matchers/CompileEventMatcher.cs
@@ -34,7 +34,7 @@ namespace UnrealBuildTool.Matchers
 			@"\)";
 
 		const string VisualCppSeverity =
-			@"(?<severity>fatal error|error|warning)(?: (?<code>[A-Z]+[0-9]+))?";
+			@"(?<severity>fatal error|[Ee]rror|[Ww]arning)(?: (?<code>[A-Z]+[0-9]+))?"; // "E"rror/"W"arning are from UHT
 
 		const string ClangLocationPattern =
 			@":" +
@@ -49,7 +49,7 @@ namespace UnrealBuildTool.Matchers
 		static readonly Regex s_preludePattern = new Regex(@"^\s*(?:In (member )?function|In file included from)");
 		static readonly Regex s_preludeFilePattern = new Regex($"^\\s*In file included from {FilePattern}:");
 		static readonly Regex s_blankLinePattern = new Regex(@"^\s*$");
-		static readonly Regex s_errorWarningPattern = new Regex("error|warning");
+		static readonly Regex s_errorWarningPattern = new Regex("[Ee]rror|[Ww]arning");
 		static readonly Regex s_clangDiagnosticPattern = new Regex($"^\\s*{FilePattern}\\s*{ClangLocationPattern}:\\s*{ClangSeverity}\\s*:");
 		static readonly Regex s_clangNotePattern = new Regex($"^\\s*{FilePattern}\\s*{ClangLocationPattern}:\\s*note:");
 		static readonly Regex s_clangMarkerPattern = new Regex(@"^(\s*)\^\s*~*$");
