Match *.hpp files as source files for compile errors.

diff --git a/Engine/Source/Programs/UnrealBuildTool.Tests/LogEventMatcherTests.cs b/Engine/Source/Programs/UnrealBuildTool.Tests/LogEventMatcherTests.cs
--- a/Engine/Source/Programs/UnrealBuildTool.Tests/LogEventMatcherTests.cs
+++ b/Engine/Source/Programs/UnrealBuildTool.Tests/LogEventMatcherTests.cs
@@ -88,6 +88,23 @@ namespace UnrealBuildToolTests
 			}
 		}
 
+		[TestMethod]
+		public void CompileHppEventMatcher()
+		{
+			// Visual C++ error
+			{
+				string[] lines =
+				{
+					@"D:/build/U5M+Inc/Sync/Engine/Source/ThirdParty/nanoflann/1.4.2/include\nanoflann/nanoflann.hpp(129,9): error: cannot use 'throw' with exceptions disabled",
+					@"        throw std::logic_error(""Try to change the size of a std::array."");",
+					@"        ^",
+				};
+
+				List<LogEvent> logEvents = Parse(String.Join("\n", lines));
+				CheckEventGroup(logEvents, 0, 3, LogLevel.Error, KnownLogEvents.Compiler);
+			}
+		}
+
 		[TestMethod]
 		public void MicrosoftEventMatcher()
 		{
diff --git a/Engine/Source/Programs/UnrealBuildTool.Tests/UnrealBuildTool.Tests.csproj b/Engine/Source/Programs/UnrealBuildTool.Tests/UnrealBuildTool.Tests.csproj
--- a/Engine/Source/Programs/UnrealBuildTool.Tests/UnrealBuildTool.Tests.csproj
+++ b/Engine/Source/Programs/UnrealBuildTool.Tests/UnrealBuildTool.Tests.csproj
@@ -2,7 +2,7 @@
 
   <PropertyGroup>
     <TargetFramework>net6.0</TargetFramework>
-
+    <Optimize>False</Optimize>
     <IsPackable>false</IsPackable>
     <OutputPath>..\..\..\Binaries\DotNET\UnrealBuildTool.Tests</OutputPath>
   </PropertyGroup>
diff --git a/Engine/Source/Programs/UnrealBuildTool/Matchers/CompileEventMatcher.cs b/Engine/Source/Programs/UnrealBuildTool/Matchers/CompileEventMatcher.cs
--- a/Engine/Source/Programs/UnrealBuildTool/Matchers/CompileEventMatcher.cs
+++ b/Engine/Source/Programs/UnrealBuildTool/Matchers/CompileEventMatcher.cs
@@ -23,7 +23,7 @@ namespace UnrealBuildTool.Matchers
 				// any path character
 				@"[^:<>*?""]+" +
 				// valid source file extension
-				@"\.(?:(?i)(?:h|c|cc|cpp|inc|inl|cs|targets))" +
+				@"\.(?:(?i)(?:h|hpp|hxx|c|cc|cpp|cxx|inc|inl|cs|targets))" +
 				// or the string "<scratch space>"
 			@")|<scratch space>)";
 
