UBT: Allow matching files without extensions as C++ source files (eg. STL headers).

diff --git a/Engine/Source/Programs/UnrealBuildTool.Tests/LogEventMatcherTests.cs b/Engine/Source/Programs/UnrealBuildTool.Tests/LogEventMatcherTests.cs
--- a/Engine/Source/Programs/UnrealBuildTool.Tests/LogEventMatcherTests.cs
+++ b/Engine/Source/Programs/UnrealBuildTool.Tests/LogEventMatcherTests.cs
@@ -413,6 +413,134 @@ namespace UnrealBuildToolTests
 			Assert.AreEqual(LogValueType.SourceFile, fileProperty.Type);
 		}
 
+		[TestMethod]
+		public void ClangEventMatcher8()
+		{
+			string[] lines =
+			{
+				@"[332/1772] Analyze graph_viewer.cc",
+
+				// group 1
+				@"In file included from Engine/Plugins/Experimental/NNI/Source/ThirdParty/ORT_4_1/Private/core/graph/graph_viewer.cc:4:",
+				@"In file included from Engine/Plugins/Experimental/NNI/Source/ThirdParty/ORT_4_1/Internal/core/graph/graph_viewer.h:6:",
+				@"In file included from Engine/Plugins/Experimental/NNI/Source/ThirdParty/ORT_4_1/Internal/core/graph/graph.h:31:",
+				@"In file included from ../Plugins/Experimental/NNI/Source/ThirdParty/Deps/gsl/gsl:31:",
+				@"Engine/Plugins/Experimental/NNI/Source/ThirdParty/Deps/gsl/gsl-lite.hpp:3217:12: warning: Forwarding reference passed to std::move(), which may unexpectedly cause lvalues to be moved; use std::forward() instead [bugprone-move-forwarding-reference]",
+				@"    return std::move( p );",
+				@"           ^",
+				@"Engine/Plugins/Experimental/NNI/Source/ThirdParty/Deps/gsl/gsl-lite.hpp:3217:12: note: Forwarding reference passed to std::move(), which may unexpectedly cause lvalues to be moved; use std::forward() instead",
+
+				// group 2
+				@"Engine/Plugins/Experimental/NNI/Source/ThirdParty/ORT_4_1/Private/core/graph/graph_viewer.cc:89:26: warning: Access to field 'nodes' results in a dereference of a null pointer (loaded from variable 'filter_info') [core.NullDereference]",
+				@"    for (NodeIndex idx : filter_info->nodes) {",
+				@"                         ^",
+				@"Engine/Plugins/Experimental/NNI/Source/ThirdParty/ORT_4_1/Private/core/graph/graph_viewer.cc:45:26: note: Passing null pointer value via 2nd parameter 'filter_info'",
+				@"    : GraphViewer(graph, nullptr) {",
+				@"                         ^~~~~~~",
+				@"Engine/Plugins/Experimental/NNI/Source/ThirdParty/ORT_4_1/Private/core/graph/graph_viewer.cc:45:7: note: Calling constructor for 'GraphViewer'",
+				@"    : GraphViewer(graph, nullptr) {",
+				@"      ^~~~~~~~~~~~~~~~~~~~~~~~~~~",
+				@"Engine/Plugins/Experimental/NNI/Source/ThirdParty/ORT_4_1/Private/core/graph/graph_viewer.cc:56:11: note: 'filter_info' is null",
+				@"          filter_info ? [this](NodeIndex idx) { return filtered_node_indices_.count(idx) == 0; }",
+				@"          ^~~~~~~~~~~",
+				@"Engine/Plugins/Experimental/NNI/Source/ThirdParty/ORT_4_1/Private/core/graph/graph_viewer.cc:56:11: note: '?' condition is false",
+				@"Engine/Plugins/Experimental/NNI/Source/ThirdParty/ORT_4_1/Private/core/graph/graph_viewer.cc:83:7: note: Calling '~function'",
+				@"      },",
+				@"      ^",
+				@"/Applications/Xcode-v13.2.1.app/Contents/Developer/Platforms/MacOSX.platform/Developer/SDKs/MacOSX.sdk/usr/include/c++/v1/functional:2547:43: note: Calling '~__value_func'",
+				@"function<_Rp(_ArgTypes...)>::~function() {}",
+				@"                                          ^",
+				@"/Applications/Xcode-v13.2.1.app/Contents/Developer/Platforms/MacOSX.platform/Developer/SDKs/MacOSX.sdk/usr/include/c++/v1/functional:1843:9: note: Taking true branch",
+				@"        if ((void*)__f_ == &__buf_)",
+				@"        ^",
+				@"/Applications/Xcode-v13.2.1.app/Contents/Developer/Platforms/MacOSX.platform/Developer/SDKs/MacOSX.sdk/usr/include/c++/v1/functional:1844:13: note: Calling '__func::destroy'",
+				@"            __f_->destroy();",
+				@"            ^~~~~~~~~~~~~~~",
+				@"/Applications/Xcode-v13.2.1.app/Contents/Developer/Platforms/MacOSX.platform/Developer/SDKs/MacOSX.sdk/usr/include/c++/v1/functional:1714:5: note: Calling '__alloc_func::destroy'",
+				@"    __f_.destroy();",
+				@"    ^~~~~~~~~~~~~~",
+				@"/Applications/Xcode-v13.2.1.app/Contents/Developer/Platforms/MacOSX.platform/Developer/SDKs/MacOSX.sdk/usr/include/c++/v1/functional:1577:32: note: Value assigned to field 'filter_info_', which participates in a condition later",
+				@"    void destroy() _NOEXCEPT { __f_.~__compressed_pair<_Target, _Alloc>(); }",
+				@"                               ^~~~~~~~~~~~~~~~~~~~~~~~~~~~~~~~~~~~~~~~~~",
+				@"/Applications/Xcode-v13.2.1.app/Contents/Developer/Platforms/MacOSX.platform/Developer/SDKs/MacOSX.sdk/usr/include/c++/v1/functional:1714:5: note: Returning from '__alloc_func::destroy'",
+				@"    __f_.destroy();",
+				@"    ^~~~~~~~~~~~~~",
+				@"/Applications/Xcode-v13.2.1.app/Contents/Developer/Platforms/MacOSX.platform/Developer/SDKs/MacOSX.sdk/usr/include/c++/v1/functional:1844:13: note: Returning from '__func::destroy'",
+				@"            __f_->destroy();",
+				@"            ^~~~~~~~~~~~~~~",
+				@"/Applications/Xcode-v13.2.1.app/Contents/Developer/Platforms/MacOSX.platform/Developer/SDKs/MacOSX.sdk/usr/include/c++/v1/functional:2547:43: note: Returning from '~__value_func'",
+				@"function<_Rp(_ArgTypes...)>::~function() {}",
+				@"                                          ^",
+				@"[REDACTED_USER_HOME]/Build/++UE5/Sync/Engine/Plugins/Experimental/NNI/Source/ThirdParty/ORT_4_1/Private/core/graph/graph_viewer.cc:83:7: note: Returning from '~function'",
+				@"      },",
+				@"      ^",
+				@"[REDACTED_USER_HOME]/Build/++UE5/Sync/Engine/Plugins/Experimental/NNI/Source/ThirdParty/ORT_4_1/Private/core/graph/graph_viewer.cc:87:7: note: Assuming field 'filter_info_' is non-null",
+				@"  if (filter_info_) {",
+				@"      ^~~~~~~~~~~~",
+				@"[REDACTED_USER_HOME]/Build/++UE5/Sync/Engine/Plugins/Experimental/NNI/Source/ThirdParty/ORT_4_1/Private/core/graph/graph_viewer.cc:87:3: note: Taking true branch",
+				@"  if (filter_info_) {",
+				@"  ^",
+				@"[REDACTED_USER_HOME]/Build/++UE5/Sync/Engine/Plugins/Experimental/NNI/Source/ThirdParty/ORT_4_1/Private/core/graph/graph_viewer.cc:89:26: note: Access to field 'nodes' results in a dereference of a null pointer (loaded from variable 'filter_info')",
+				@"    for (NodeIndex idx : filter_info->nodes) {",
+				@"                         ^~~~~~~~~~~",
+
+				// group 3
+				@"Engine/Plugins/Experimental/NNI/Source/ThirdParty/ORT_4_1/Private/core/graph/graph_viewer.cc:119:5: warning: Method called on moved-from object 'nodes_in_topological_order_' of type 'std::vector' [cplusplus.Move]",
+				@"    nodes_in_topological_order_.reserve(filter_info->nodes.size());",
+				@"    ^",
+				@"Engine/Plugins/Experimental/NNI/Source/ThirdParty/ORT_4_1/Private/core/graph/graph_viewer.cc:49:7: note: Calling constructor for 'GraphViewer'",
+				@"    : GraphViewer(graph, &filter_info) {",
+				@"      ^~~~~~~~~~~~~~~~~~~~~~~~~~~~~~~~",
+				@"Engine/Plugins/Experimental/NNI/Source/ThirdParty/ORT_4_1/Private/core/graph/graph_viewer.cc:56:11: note: 'filter_info' is non-null",
+				@"          filter_info ? [this](NodeIndex idx) { return filtered_node_indices_.count(idx) == 0; }",
+				@"          ^~~~~~~~~~~",
+				@"Engine/Plugins/Experimental/NNI/Source/ThirdParty/ORT_4_1/Private/core/graph/graph_viewer.cc:56:11: note: '?' condition is true",
+				@"Engine/Plugins/Experimental/NNI/Source/ThirdParty/ORT_4_1/Private/core/graph/graph_viewer.cc:87:7: note: Assuming field 'filter_info_' is non-null",
+				@"  if (filter_info_) {",
+				@"      ^~~~~~~~~~~~",
+				@"Engine/Plugins/Experimental/NNI/Source/ThirdParty/ORT_4_1/Private/core/graph/graph_viewer.cc:87:3: note: Taking true branch",
+				@"  if (filter_info_) {",
+				@"  ^",
+				@"Engine/Plugins/Experimental/NNI/Source/ThirdParty/ORT_4_1/Private/core/graph/graph_viewer.cc:118:23: note: Object 'nodes_in_topological_order_' of type 'std::vector' is left in a valid but unspecified state after move",
+				@"    auto orig_order = std::move(nodes_in_topological_order_);",
+				@"                      ^~~~~~~~~~~~~~~~~~~~~~~~~~~~~~~~~~~~~~",
+				@"Engine/Plugins/Experimental/NNI/Source/ThirdParty/ORT_4_1/Private/core/graph/graph_viewer.cc:119:5: note: Method called on moved-from object 'nodes_in_topological_order_' of type 'std::vector'",
+				@"    nodes_in_topological_order_.reserve(filter_info->nodes.size());",
+				@"    ^~~~~~~~~~~~~~~~~~~~~~~~~~~~~~~~~~~~~~~~~~~~~~~~~~~~~~~~~~~~~~",
+
+				// group 4
+				@"Engine/Plugins/Experimental/NNI/Source/ThirdParty/ORT_4_1/Private/core/graph/graph_viewer.cc:145:5: warning: Method called on moved-from object 'nodes_in_topological_order_with_priority_' of type 'std::vector' [cplusplus.Move]",
+				@"    nodes_in_topological_order_with_priority_.reserve(filter_info->nodes.size());",
+				@"    ^",
+				@"Engine/Plugins/Experimental/NNI/Source/ThirdParty/ORT_4_1/Private/core/graph/graph_viewer.cc:49:7: note: Calling constructor for 'GraphViewer'",
+				@"    : GraphViewer(graph, &filter_info) {",
+				@"      ^~~~~~~~~~~~~~~~~~~~~~~~~~~~~~~~",
+				@"Engine/Plugins/Experimental/NNI/Source/ThirdParty/ORT_4_1/Private/core/graph/graph_viewer.cc:56:11: note: 'filter_info' is non-null",
+				@"          filter_info ? [this](NodeIndex idx) { return filtered_node_indices_.count(idx) == 0; }",
+				@"          ^~~~~~~~~~~",
+				@"Engine/Plugins/Experimental/NNI/Source/ThirdParty/ORT_4_1/Private/core/graph/graph_viewer.cc:56:11: note: '?' condition is true",
+				@"Engine/Plugins/Experimental/NNI/Source/ThirdParty/ORT_4_1/Private/core/graph/graph_viewer.cc:87:7: note: Assuming field 'filter_info_' is non-null",
+				@"  if (filter_info_) {",
+				@"      ^~~~~~~~~~~~",
+				@"Engine/Plugins/Experimental/NNI/Source/ThirdParty/ORT_4_1/Private/core/graph/graph_viewer.cc:87:3: note: Taking true branch",
+				@"  if (filter_info_) {",
+				@"  ^",
+				@"Engine/Plugins/Experimental/NNI/Source/ThirdParty/ORT_4_1/Private/core/graph/graph_viewer.cc:144:32: note: Object 'nodes_in_topological_order_with_priority_' of type 'std::vector' is left in a valid but unspecified state after move",
+				@"    auto orig_priority_order = std::move(nodes_in_topological_order_with_priority_);",
+				@"                               ^~~~~~~~~~~~~~~~~~~~~~~~~~~~~~~~~~~~~~~~~~~~~~~~~~~~",
+				@"Engine/Plugins/Experimental/NNI/Source/ThirdParty/ORT_4_1/Private/core/graph/graph_viewer.cc:145:5: note: Method called on moved-from object 'nodes_in_topological_order_with_priority_' of type 'std::vector'",
+				@"    nodes_in_topological_order_with_priority_.reserve(filter_info->nodes.size());",
+				@"    ^~~~~~~~~~~~~~~~~~~~~~~~~~~~~~~~~~~~~~~~~~~~~~~~~~~~~~~~~~~~~~~~~~~~~~~~~~~~",
+				@"4 warnings generated"
+			};
+
+			List<LogEvent> logEvents = Parse(String.Join("\n", lines));
+			CheckEventGroup(logEvents.Slice(0, 8), 1, 8, LogLevel.Warning, KnownLogEvents.Compiler);
+			CheckEventGroup(logEvents.Slice(8, 52), 9, 52, LogLevel.Warning, KnownLogEvents.Compiler);
+			CheckEventGroup(logEvents.Slice(60, 22), 61, 22, LogLevel.Warning, KnownLogEvents.Compiler);
+			CheckEventGroup(logEvents.Slice(82, 22), 83, 22, LogLevel.Warning, KnownLogEvents.Compiler);
+		}
+
 		[TestMethod]
 		public void IOSCompileErrorMatcher()
 		{
diff --git a/Engine/Source/Programs/UnrealBuildTool/Matchers/CompileEventMatcher.cs b/Engine/Source/Programs/UnrealBuildTool/Matchers/CompileEventMatcher.cs
--- a/Engine/Source/Programs/UnrealBuildTool/Matchers/CompileEventMatcher.cs
+++ b/Engine/Source/Programs/UnrealBuildTool/Matchers/CompileEventMatcher.cs
@@ -22,8 +22,8 @@ namespace UnrealBuildTool.Matchers
 				@"[^:(\s]+" +
 				// any path character
 				@"[^:<>*?""]+" +
-				// valid source file extension
-				@"\.(?:(?i)(?:h|hpp|hxx|c|cc|cpp|cxx|inc|inl|cs|targets))" +
+				// valid source file extension (or extensionless file)
+				@"(?:\.(?i)(?:h|hpp|hxx|c|cc|cpp|cxx|inc|inl|cs|targets)|[/\\][a-zA-Z0-9]+)" +
 				// or the string "<scratch space>"
 			@")|<scratch space>)";
 
