UBT: Improve matcher for Clang line markers.

diff --git a/Engine/Source/Programs/UnrealBuildTool.Tests/LogEventMatcherTests.cs b/Engine/Source/Programs/UnrealBuildTool.Tests/LogEventMatcherTests.cs
--- a/Engine/Source/Programs/UnrealBuildTool.Tests/LogEventMatcherTests.cs
+++ b/Engine/Source/Programs/UnrealBuildTool.Tests/LogEventMatcherTests.cs
@@ -541,6 +541,22 @@ namespace UnrealBuildToolTests
 			CheckEventGroup(logEvents.Slice(82, 22), 83, 22, LogLevel.Warning, KnownLogEvents.Compiler);
 		}
 
+		[TestMethod]
+		public void ClangEventMatcher9()
+		{
+			string[] lines =
+			{
+				@"[421/1292] Compile Module.ApplicationCore.cpp",
+				@"In file included from Engine/Intermediate/Build/IOS/UnrealGame/Development/ApplicationCore/Module.ApplicationCore.cpp:14:",
+				@"Engine/Source/Runtime/ApplicationCore/Private/IOS/Accessibility/IOSAccessibilityCache.cpp:119:26: error: implicit conversion loses floating-point precision: 'CGFloat' (aka 'double') to 'const float' [-Werror,-Wimplicit-float-conversion]",
+				@"                                        const float Scale = [IOSAppDelegate GetDelegate].IOSView.contentScaleFactor;",
+				@"                                                    ~~~~~   ^~~~~~~~~~~~~~~~~~~~~~~~~~~~~~~~~~~~~~~~~~~~~~~~~~~~~~~",
+			};
+
+			List<LogEvent> logEvents = Parse(String.Join("\n", lines));
+			CheckEventGroup(logEvents, 0, 5, LogLevel.Error, KnownLogEvents.Compiler);
+		}
+
 		[TestMethod]
 		public void IOSCompileErrorMatcher()
 		{
diff --git a/Engine/Source/Programs/UnrealBuildTool/Matchers/CompileEventMatcher.cs b/Engine/Source/Programs/UnrealBuildTool/Matchers/CompileEventMatcher.cs
--- a/Engine/Source/Programs/UnrealBuildTool/Matchers/CompileEventMatcher.cs
+++ b/Engine/Source/Programs/UnrealBuildTool/Matchers/CompileEventMatcher.cs
@@ -52,7 +52,7 @@ namespace UnrealBuildTool.Matchers
 		static readonly Regex s_errorWarningPattern = new Regex("[Ee]rror|[Ww]arning");
 		static readonly Regex s_clangDiagnosticPattern = new Regex($"^\\s*{FilePattern}\\s*{ClangLocationPattern}:\\s*{ClangSeverity}\\s*:");
 		static readonly Regex s_clangNotePattern = new Regex($"^\\s*{FilePattern}\\s*{ClangLocationPattern}:\\s*note:");
-		static readonly Regex s_clangMarkerPattern = new Regex(@"^(\s*)\^\s*~*$");
+		static readonly Regex s_clangMarkerPattern = new Regex(@"^(\s*)[\^~][\s\^~]*$");
 
 		static readonly string[] s_invalidExtensions =
 		{
