Horde: Add log event matcher for XGE cache warnings.

diff --git a/Engine/Source/Programs/Shared/EpicGames.Core/KnownLogEvents.cs b/Engine/Source/Programs/Shared/EpicGames.Core/KnownLogEvents.cs
--- a/Engine/Source/Programs/Shared/EpicGames.Core/KnownLogEvents.cs
+++ b/Engine/Source/Programs/Shared/EpicGames.Core/KnownLogEvents.cs
@@ -185,6 +185,11 @@ namespace EpicGames.Core
 		/// </summary>
 		public static EventId Systemic_Xge_BuildFailed { get; } = new EventId(713);
 
+		/// <summary>
+		/// Cache size reached
+		/// </summary>
+		public static EventId Systemic_Xge_CacheLimit { get; } = new EventId(714);
+
 		/// <summary>
 		/// DDC is slow
 		/// </summary>
diff --git a/Engine/Source/Programs/UnrealBuildTool.Tests/LogEventMatcherTests.cs b/Engine/Source/Programs/UnrealBuildTool.Tests/LogEventMatcherTests.cs
--- a/Engine/Source/Programs/UnrealBuildTool.Tests/LogEventMatcherTests.cs
+++ b/Engine/Source/Programs/UnrealBuildTool.Tests/LogEventMatcherTests.cs
@@ -884,6 +884,29 @@ namespace UnrealBuildToolTests
 			CheckEventGroup(logEvents.Slice(7, 1), 8, 1, LogLevel.Error, KnownLogEvents.Systemic_Xge_BuildFailed);
 		}
 
+		[TestMethod]
+		public void XoreaxErrorMatcher2()
+		{
+			string[] lines =
+			{
+				@"--------------------Build Cache summary----------------------------------------",
+				@"Build Tasks: 8042",
+				@"Build Cache Efficiency: 0% (0 tasks)",
+				@"New cache size: 17 GB",
+				@"Updated 352 items (20 GB)",
+				@"WARNING: 389 items (19.4 GB) removed from the cache due to reaching the cache size limit",
+				@"",
+				@"1 build system warning(s):",
+				@"   - Build Cache performance hit",
+				@"",
+				@"Took 2515.0181741s to run xgConsole.exe, ExitCode=0",
+			};
+
+			List<LogEvent> logEvents = Parse(lines);
+			Assert.AreEqual(1, logEvents.Count);
+			CheckEventGroup(logEvents, 5, 1, LogLevel.Information, KnownLogEvents.Systemic_Xge_CacheLimit);
+		}
+
 		[TestMethod]
 		public void UhtErrorMatcher()
 		{
diff --git a/Engine/Source/Programs/UnrealBuildTool/Matchers/XoreaxEventMatcher.cs b/Engine/Source/Programs/UnrealBuildTool/Matchers/XoreaxEventMatcher.cs
--- a/Engine/Source/Programs/UnrealBuildTool/Matchers/XoreaxEventMatcher.cs
+++ b/Engine/Source/Programs/UnrealBuildTool/Matchers/XoreaxEventMatcher.cs
@@ -15,8 +15,16 @@ namespace UnrealBuildTool.Matchers
 		static readonly Regex s_buildSystemWarning = new Regex(@"^\s*--------------------Build System Warning[- ]");
 		static readonly Regex s_buildSystemWarningNext = new Regex(@"^(\s*)([^ ].*):");
 
+		static readonly Regex s_cacheWarning = new Regex(@"^\s*WARNING: \d+ items \([^\)]*\) removed from the cache due to reaching the cache size limit");
+
 		public LogEventMatch? Match(ILogCursor cursor)
 		{
+			if (cursor.IsMatch(s_cacheWarning))
+			{
+				LogEventBuilder builder = new LogEventBuilder(cursor);
+				return builder.ToMatch(LogEventPriority.High, LogLevel.Information, KnownLogEvents.Systemic_Xge_CacheLimit);
+			}
+
 			if (cursor.IsMatch(s_buildService))
 			{
 				LogEventBuilder builder = new LogEventBuilder(cursor);
