UBT: Add a matcher for XGE errors.

diff --git a/Engine/Source/Programs/UnrealBuildTool.Tests/LogEventMatcherTests.cs b/Engine/Source/Programs/UnrealBuildTool.Tests/LogEventMatcherTests.cs
--- a/Engine/Source/Programs/UnrealBuildTool.Tests/LogEventMatcherTests.cs
+++ b/Engine/Source/Programs/UnrealBuildTool.Tests/LogEventMatcherTests.cs
@@ -20,6 +20,11 @@ namespace UnrealBuildToolTests
 
 		class LoggerCapture : ILogger
 		{
+			class Scope : IDisposable
+			{
+				public void Dispose() { }
+			}
+
 			int _logLineIndex;
 
 			public void Reset()
@@ -30,11 +35,11 @@ namespace UnrealBuildToolTests
 
 			public List<LogEvent> _events = new List<LogEvent>();
 
-			public IDisposable? BeginScope<TState>(TState state) => null;
+			public IDisposable BeginScope<TState>(TState state) => new Scope();
 
 			public bool IsEnabled(LogLevel logLevel) => true;
 
-			public void Log<TState>(LogLevel logLevel, EventId eventId, TState state, Exception exception, Func<TState, Exception?, string> formatter)
+			public void Log<TState>(LogLevel logLevel, EventId eventId, TState state, Exception? exception, Func<TState, Exception?, string> formatter)
 			{
 				LogEvent logEvent = LogEvent.Read(JsonLogEvent.FromLoggerState(logLevel, eventId, state, exception, formatter).Data.Span);
 				if (logEvent.Level != LogLevel.Information || logEvent.Id != default || logEvent.Properties != null)
@@ -907,6 +912,30 @@ namespace UnrealBuildToolTests
 			CheckEventGroup(logEvents, 5, 1, LogLevel.Information, KnownLogEvents.Systemic_Xge_CacheLimit);
 		}
 
+		[TestMethod]
+		public void XoreaxErrorMatcher3()
+		{
+			string[] lines =
+			{
+				@"Module.ExternalSource.cpp (Agent '10-99-199-46 (Core #3)', 0:16.43 at +23:42)",
+				@"Default.rc2 (0:00.29 at +23:59)",
+				@"--------------------Build System Error-----------------------------------------",
+				@"Fatal error:",
+				@"    Task queue management failed.",
+				@"    Error starting Task 'clang-cl: Task: Env_0->Action1715_1 (Tool1715_1)' on machine 'Local CPU 8'",
+				@"    Failed to map view of file mapping 0x082C at position 0:0 (0 bytes): Not enough memory resources are available to process this command (8)",
+				@"-------------------------------------------------------------------------------",
+				@"",
+				@"--------------------Build Cache summary----------------------------------------",
+				@"Build Tasks: 15751",
+				@"Build Cache Efficiency: 0% (0 tasks)"
+			};
+
+			List<LogEvent> logEvents = Parse(lines);
+			Assert.AreEqual(6, logEvents.Count);
+			CheckEventGroup(logEvents, 2, 6, LogLevel.Error, KnownLogEvents.Systemic_Xge);
+		}
+
 		[TestMethod]
 		public void UhtErrorMatcher()
 		{
diff --git a/Engine/Source/Programs/UnrealBuildTool/Matchers/XoreaxEventMatcher.cs b/Engine/Source/Programs/UnrealBuildTool/Matchers/XoreaxEventMatcher.cs
--- a/Engine/Source/Programs/UnrealBuildTool/Matchers/XoreaxEventMatcher.cs
+++ b/Engine/Source/Programs/UnrealBuildTool/Matchers/XoreaxEventMatcher.cs
@@ -12,6 +12,9 @@ namespace UnrealBuildTool.Matchers
 
 		static readonly Regex s_xgConsole = new Regex(@"BUILD FAILED: (.*)xgConsole\.exe(.*)");
 
+		static readonly Regex s_buildSystemError = new Regex(@"^\s*--------------------Build System Error");
+		static readonly Regex s_buildSystemErrorEnd = new Regex(@"^(\s*)--------------------");
+
 		static readonly Regex s_buildSystemWarning = new Regex(@"^\s*--------------------Build System Warning[- ]");
 		static readonly Regex s_buildSystemWarningNext = new Regex(@"^(\s*)([^ ].*):");
 
@@ -37,6 +40,23 @@ namespace UnrealBuildTool.Matchers
 				return builder.ToMatch(LogEventPriority.High, LogLevel.Error, KnownLogEvents.Systemic_Xge_BuildFailed);
 			}
 
+			if (cursor.IsMatch(s_buildSystemError))
+			{
+				LogEventBuilder builder = new LogEventBuilder(cursor);
+				builder.MoveNext();
+
+				for(int idx = 0; idx < 100; idx++)
+				{
+					if (builder.Current.IsMatch(idx, s_buildSystemErrorEnd))
+					{
+						builder.MoveNext(idx);
+						break;
+					}
+				}
+
+				return builder.ToMatch(LogEventPriority.High, LogLevel.Error, KnownLogEvents.Systemic_Xge);
+			}
+
 			if (cursor.IsMatch(s_buildSystemWarning))
 			{
 				LogEventBuilder builder = new LogEventBuilder(cursor);
