UBT: Add a matcher for Verse errors.

diff --git a/Engine/Source/Programs/UnrealBuildTool.Tests/LogEventMatcherTests.cs b/Engine/Source/Programs/UnrealBuildTool.Tests/LogEventMatcherTests.cs
--- a/Engine/Source/Programs/UnrealBuildTool.Tests/LogEventMatcherTests.cs
+++ b/Engine/Source/Programs/UnrealBuildTool.Tests/LogEventMatcherTests.cs
@@ -956,6 +956,46 @@ namespace UnrealBuildToolTests
 			Assert.AreEqual(LogValueType.SourceFile, fileProperty.Type);
 		}
 
+		[TestMethod]
+		public void UhtErrorMatcher2()
+		{
+			string[] lines =
+			{
+				@"Parsing headers for UnrealEditor",
+				@"  Running Internal UnrealHeaderTool UnrealEditor [REDACTED_USER_HOME]/Build/U5M+Inc/Sync/Engine/Intermediate/Build/Mac/x86_64/UnrealEditor/Debug/UnrealEditor.uhtmanifest -WarningsAsErrors",
+				@"[REDACTED_USER_HOME]/Build/U5M+Inc/Sync/Engine/Restricted/NotForLicensees/Plugins/Compositor/Source/Compositor/Public/Assets/Composite.h(88): Error: Native pointer usage in member declaration detected [[[UComposite*]]].  This is disallowed for the target/module, consider TObjectPtr as an alternative.",
+				@"[REDACTED_USER_HOME]/Build/U5M+Inc/Sync/Engine/Restricted/NotForLicensees/Plugins/Compositor/Source/Compositor/Public/Components/CompositeCaptureComponent2D.h(53): Error: Native pointer usage in member declaration detected [[[APlayerCameraManager*]]].  This is disallowed for the target/module, consider TObjectPtr as an alternative.",
+				@"[REDACTED_USER_HOME]/Build/U5M+Inc/Sync/Engine/Restricted/NotForLicensees/Plugins/Compositor/Source/Compositor/Public/Assets/Composite.h(92): Error: Native pointer usage in member declaration detected [[[UComposite*]]].  This is disallowed for the target/module, consider TObjectPtr as an alternative.",
+			};
+
+			List<LogEvent> logEvents = Parse(lines);
+			Assert.AreEqual(3, logEvents.Count);
+			CheckEventGroup(logEvents.Slice(0, 1), 2, 1, LogLevel.Error, KnownLogEvents.Compiler);
+			CheckEventGroup(logEvents.Slice(1, 1), 3, 1, LogLevel.Error, KnownLogEvents.Compiler);
+			CheckEventGroup(logEvents.Slice(2, 1), 4, 1, LogLevel.Error, KnownLogEvents.Compiler);
+
+			LogValue fileProperty = (LogValue)logEvents[0].GetProperty("file");
+			Assert.AreEqual(@"[REDACTED_USER_HOME]/Build/U5M+Inc/Sync/Engine/Restricted/NotForLicensees/Plugins/Compositor/Source/Compositor/Public/Assets/Composite.h", fileProperty.Text);
+			Assert.AreEqual(LogValueType.SourceFile, fileProperty.Type);
+		}
+
+		[TestMethod]
+		public void VerseErrorMatcher()
+		{
+			string[] lines =
+			{
+				@"C:/Horde/TestGame/Plugins/TestPlugin/Test.verse(9,9, 9,43): Verse compiler error V3587: Something went wrong"
+			};
+
+			List<LogEvent> logEvents = Parse(lines);
+			Assert.AreEqual(1, logEvents.Count);
+			CheckEventGroup(logEvents, 0, 1, LogLevel.Error, KnownLogEvents.Compiler);
+
+			LogValue fileProperty = (LogValue)logEvents[0].GetProperty("file");
+			Assert.AreEqual(@"C:/Horde/TestGame/Plugins/TestPlugin/Test.verse", fileProperty.Text);
+			Assert.AreEqual(LogValueType.SourceFile, fileProperty.Type);
+		}
+
 		static List<LogEvent> Parse(IEnumerable<string> lines)
 		{
 			return Parse(String.Join("\n", lines));
diff --git a/Engine/Source/Programs/UnrealBuildTool/Matchers/CompileEventMatcher.cs b/Engine/Source/Programs/UnrealBuildTool/Matchers/CompileEventMatcher.cs
--- a/Engine/Source/Programs/UnrealBuildTool/Matchers/CompileEventMatcher.cs
+++ b/Engine/Source/Programs/UnrealBuildTool/Matchers/CompileEventMatcher.cs
@@ -23,18 +23,27 @@ namespace UnrealBuildTool.Matchers
 				// any path character
 				@"[^:<>*?""]+" +
 				// valid source file extension (or extensionless file)
-				@"(?:\.(?i)(?:h|hpp|hxx|c|cc|cpp|cxx|inc|inl|cs|targets)|[/\\][a-zA-Z0-9]+)" +
+				@"(?:\.(?i)(?:h|hpp|hxx|c|cc|cpp|cxx|inc|inl|cs|targets|verse)|[/\\][a-zA-Z0-9]+)" +
 				// or the string "<scratch space>"
 			@")|<scratch space>)";
 
 		const string VisualCppLocationPattern =
 			@"\(" +
-				@"(?<line>\d+)" + // line number
-				@"(?:,(?<column>\d+))?" + // optional column number
+				@"(?:" + 
+					@"(?<line>\d+)" + // (line)
+					@"|" +
+					@"(?:(?<line>\d+)-(?<maxline>\d+))" + // (line-line)
+					@"|" +
+					@"(?:(?<line>\d+),(?<column>\d+))" + // (line,col)
+					@"|" +
+					@"(?:(?<line>\d+),(?<column>\d+)-(?<maxcolumn>\d+))" + // (line,col-col)
+					@"|" +
+					@"(?:(?<line>\d+),(?<column>\d+),\s*(?<maxline>\d+),(?<maxcolumn>\d+))" + // (line,col,line,col)
+				@")" +
 			@"\)";
 
 		const string VisualCppSeverity =
-			@"(?<severity>fatal error|[Ee]rror|[Ww]arning)(?: (?<code>[A-Z]+[0-9]+))?"; // "E"rror/"W"arning are from UHT
+			@"(?:Verse compiler )?(?<severity>fatal error|[Ee]rror|[Ww]arning)(?: (?<code>[A-Z]+[0-9]+))?"; // "E"rror/"W"arning are from UHT
 
 		const string ClangLocationPattern =
 			@":" +
@@ -243,6 +252,8 @@ namespace UnrealBuildTool.Matchers
 			builder.AnnotateSourceFile(match.Groups["file"], sourceFileBaseDir);
 			builder.TryAnnotate(match.Groups["line"], LogEventMarkup.LineNumber);
 			builder.TryAnnotate(match.Groups["column"], LogEventMarkup.ColumnNumber);
+			builder.TryAnnotate(match.Groups["maxline"], LogEventMarkup.LineNumber);
+			builder.TryAnnotate(match.Groups["maxcolumn"], LogEventMarkup.ColumnNumber);
 
 			string indent = ExtractIndent(builder.Current.CurrentLine ?? String.Empty);
 			string nextIndent = indent + " ";
