Horde: Add UBT matchers for additional XGE messages, with test
#rnx

diff --git a/Engine/Source/Programs/UnrealBuildTool.Tests/LogEventMatcherTests.cs b/Engine/Source/Programs/UnrealBuildTool.Tests/LogEventMatcherTests.cs
--- a/Engine/Source/Programs/UnrealBuildTool.Tests/LogEventMatcherTests.cs
+++ b/Engine/Source/Programs/UnrealBuildTool.Tests/LogEventMatcherTests.cs
@@ -936,6 +936,30 @@ namespace UnrealBuildToolTests
 			CheckEventGroup(logEvents, 2, 6, LogLevel.Error, KnownLogEvents.Systemic_Xge);
 		}
 
+		[TestMethod]
+		public void XoreaxErrorMatcher4()
+		{
+			string[] lines =
+			{
+				@"--------------------Build Cache summary----------------------------------------",
+				@"Build Tasks: 2007",
+				@"Build Cache Efficiency: 0% (0 tasks)",
+				@"New cache size: 24.9 GB",
+				@"Updated 1899 items (9.8 GB)",
+				@"WARNING: Several items removed from the cache due to reaching the cache size limit.",
+				@"WARNING: The Build Cache is close to full, limiting your ability to benefit from previously cached data to speed up your builds. It is recommended to increase the cache size to be larger than the sum of all build artifacts that are cached.",
+				@"1 build system warning(s):",
+				@"   - Build Cache performance hit",
+				@"Took 413.95489860000004s to run xgConsole.exe, ExitCode=0"
+			};
+
+			List<LogEvent> logEvents = Parse(lines);
+			Assert.AreEqual(2, logEvents.Count);
+			CheckEventGroup(logEvents.Slice(0, 1), 5, 1, LogLevel.Information, KnownLogEvents.Systemic_Xge_CacheLimit);
+			CheckEventGroup(logEvents.Slice(1, 1), 6, 1, LogLevel.Information, KnownLogEvents.Systemic_Xge_CacheLimit);
+		}
+
+
 		[TestMethod]
 		public void UhtErrorMatcher()
 		{
diff --git a/Engine/Source/Programs/UnrealBuildTool/Matchers/XoreaxEventMatcher.cs b/Engine/Source/Programs/UnrealBuildTool/Matchers/XoreaxEventMatcher.cs
--- a/Engine/Source/Programs/UnrealBuildTool/Matchers/XoreaxEventMatcher.cs
+++ b/Engine/Source/Programs/UnrealBuildTool/Matchers/XoreaxEventMatcher.cs
@@ -19,10 +19,12 @@ namespace UnrealBuildTool.Matchers
 		static readonly Regex s_buildSystemWarningNext = new Regex(@"^(\s*)([^ ].*):");
 
 		static readonly Regex s_cacheWarning = new Regex(@"^\s*WARNING: \d+ items \([^\)]*\) removed from the cache due to reaching the cache size limit");
+		static readonly Regex s_cacheWarning2 = new Regex(@"^\s*WARNING: Several items removed from the cache due to reaching the cache size limit");
+		static readonly Regex s_cacheWarning3 = new Regex(@"^\s*WARNING: The Build Cache is close to full");
 
 		public LogEventMatch? Match(ILogCursor cursor)
 		{
-			if (cursor.IsMatch(s_cacheWarning))
+			if (cursor.IsMatch(s_cacheWarning) || cursor.IsMatch(s_cacheWarning2) || cursor.IsMatch(s_cacheWarning3))
 			{
 				LogEventBuilder builder = new LogEventBuilder(cursor);
 				return builder.ToMatch(LogEventPriority.High, LogLevel.Information, KnownLogEvents.Systemic_Xge_CacheLimit);
