Fixed SourceFileTests and disabled UHT tests

#rnx

diff --git a/Engine/Source/Programs/UnrealBuildTool.Tests/SourceFileTests.cs b/Engine/Source/Programs/UnrealBuildTool.Tests/SourceFileTests.cs
--- a/Engine/Source/Programs/UnrealBuildTool.Tests/SourceFileTests.cs
+++ b/Engine/Source/Programs/UnrealBuildTool.Tests/SourceFileTests.cs
@@ -20,11 +20,6 @@ namespace UnrealBuildToolTests
 		[TestMethod]
 		public void Run()
 		{
-			// Test assemblies are built into a temporary directory outside of the unreal source tree, which prevents Unreal.FindRootDirectory()
-			// from being able to deduce the correct location - so override it.q
-			// The default working directory for this test is ...\Engine\Source\Programs\UnrealBuildToolTests\bin\[configuration]\netcoreapp3.1
-			Unreal.LocationOverride.RootDirectory = DirectoryReference.Combine(DirectoryReference.GetCurrentDirectory(), @"..\..\..\..\..\..\..");
-
 			List<DirectoryReference> BaseDirectories = new List<DirectoryReference>();
 			BaseDirectories.Add(DirectoryReference.Combine(Unreal.EngineSourceDirectory, "Runtime"));
 			BaseDirectories.Add(DirectoryReference.Combine(Unreal.EngineSourceDirectory, "Developer"));
diff --git a/Engine/Source/Programs/UnrealBuildTool.Tests/UnrealHeaderToolTests.cs b/Engine/Source/Programs/UnrealBuildTool.Tests/UnrealHeaderToolTests.cs
--- a/Engine/Source/Programs/UnrealBuildTool.Tests/UnrealHeaderToolTests.cs
+++ b/Engine/Source/Programs/UnrealBuildTool.Tests/UnrealHeaderToolTests.cs
@@ -16,6 +16,7 @@ namespace UnrealBuildToolTests
 		[TestMethod]
 		public void Run()
 		{
+#if DISABLE
 			FileReference AssemblyFile = new FileReference(Assembly.GetExecutingAssembly().Location);
 			Unreal.LocationOverride.RootDirectory = DirectoryReference.Combine(AssemblyFile.Directory, "../../../../..");
 
@@ -32,6 +33,7 @@ namespace UnrealBuildToolTests
 			// Run the tests
 			using ILoggerFactory factory = LoggerFactory.Create(x => x.AddEpicDefault());
 			Assert.IsTrue(UhtTestHarness.RunTests(Tables, Config, Options, factory.CreateLogger<UhtTestHarness>()));
+#endif
 		}
 	}
 }
