UBT: Match engine-style script errors in logs.

diff --git a/Engine/Source/Programs/UnrealBuildTool.Tests/LogEventMatcherTests.cs b/Engine/Source/Programs/UnrealBuildTool.Tests/LogEventMatcherTests.cs
--- a/Engine/Source/Programs/UnrealBuildTool.Tests/LogEventMatcherTests.cs
+++ b/Engine/Source/Programs/UnrealBuildTool.Tests/LogEventMatcherTests.cs
@@ -1018,6 +1018,27 @@ namespace UnrealBuildToolTests
 			Assert.AreEqual(LogValueType.SourceFile, fileProperty.Type);
 		}
 
+		[TestMethod]
+		public void VerseInEngineErrorMatcher()
+		{
+			string[] lines =
+			{
+				@"LogSolarisIde: Error: C:/Horde/Foo/Plugins/VerseAI/CompanionAI/Source/CompanionAI/Verse/CompanionAI.verse(651,46, 651,64): Script error 3506: Unknown identifier `sphere_draw_params`."
+			};
+
+			List<LogEvent> logEvents = Parse(String.Join("\n", lines));
+			Assert.AreEqual(1, logEvents.Count);
+			CheckEventGroup(logEvents.Slice(0, 1), 0, 1, LogLevel.Error, KnownLogEvents.Compiler);
+
+			LogEvent logEvent = logEvents[0];
+			Assert.AreEqual("3506", logEvent.GetProperty("code").ToString());
+			Assert.AreEqual(LogLevel.Error, logEvent.Level);
+
+			LogValue fileProperty = logEvents[0].GetProperty<LogValue>("file");
+			Assert.AreEqual(LogValueType.SourceFile, fileProperty.Type);
+			Assert.AreEqual(@"C:/Horde/Foo/Plugins/VerseAI/CompanionAI/Source/CompanionAI/Verse/CompanionAI.verse", fileProperty.Text);
+		}
+
 		static List<LogEvent> Parse(IEnumerable<string> lines)
 		{
 			return Parse(String.Join("\n", lines));
diff --git a/Engine/Source/Programs/UnrealBuildTool/Matchers/CompileEventMatcher.cs b/Engine/Source/Programs/UnrealBuildTool/Matchers/CompileEventMatcher.cs
--- a/Engine/Source/Programs/UnrealBuildTool/Matchers/CompileEventMatcher.cs
+++ b/Engine/Source/Programs/UnrealBuildTool/Matchers/CompileEventMatcher.cs
@@ -43,7 +43,7 @@ namespace UnrealBuildTool.Matchers
 			@"\)";
 
 		const string VisualCppSeverity =
-			@"(?:Verse compiler )?(?<severity>fatal error|[Ee]rror|[Ww]arning)(?: (?<code>[A-Z]+[0-9]+))?"; // "E"rror/"W"arning are from UHT
+			@"(?:Verse compiler |Script )?(?<severity>fatal error|[Ee]rror|[Ww]arning)(?: (?<code>[A-Z]*[0-9]+))?"; // "E"rror/"W"arning are from UHT
 
 		const string ClangLocationPattern =
 			@":" +
@@ -199,7 +199,7 @@ namespace UnrealBuildTool.Matchers
 			return false;
 		}
 
-		static readonly Regex s_msvcPattern = new Regex($"^\\s*(?:ERROR: |WARNING: )?{FilePattern}(?:{VisualCppLocationPattern})? ?:\\s+{VisualCppSeverity}:");
+		static readonly Regex s_msvcPattern = new Regex($"^\\s*(?:ERROR: |WARNING: |Log[A-Za-z_]+: [A-Z][a-z]+: )?{FilePattern}(?:{VisualCppLocationPattern})? ?:\\s+{VisualCppSeverity}:");
 		static readonly Regex s_msvcNotePattern = new Regex($"^\\s*{FilePattern}(?:{VisualCppLocationPattern})?\\s*: note:");
 		static readonly Regex s_projectPattern = new Regex(@"\[(?<project>[^[\]]+)]\s*$");
 
