Change artifact tests to be async, so callstacks are correct while debugging.

diff --git a/Engine/Source/Programs/UnrealBuildTool.Tests/ArtifactsTest.cs b/Engine/Source/Programs/UnrealBuildTool.Tests/ArtifactsTest.cs
--- a/Engine/Source/Programs/UnrealBuildTool.Tests/ArtifactsTest.cs
+++ b/Engine/Source/Programs/UnrealBuildTool.Tests/ArtifactsTest.cs
@@ -1,6 +1,8 @@
 // Copyright Epic Games, Inc. All Rights Reserved.
 
+using System.Reflection;
 using System.Threading;
+using System.Threading.Tasks;
 using EpicGames.Core;
 using Microsoft.Extensions.Logging.Abstractions;
 using Microsoft.VisualStudio.TestTools.UnitTesting;
@@ -54,28 +56,28 @@ namespace UnrealBuildToolTests
 		}
 
 		[TestMethod]
-		public void ArtifactBundleStorageTest1()
+		public async Task ArtifactBundleStorageTest1()
 		{
 			CancellationToken cancellationToken = default;
 
 			IArtifactCache cache = HordeStorageArtifactCache.CreateMemoryCache(NullLogger.Instance);
 
-			_ = cache.WaitForReadyAsync().Result;
+			await cache.WaitForReadyAsync();
 			Assert.AreEqual(ArtifactCacheState.Available, cache.State);
 
 			ArtifactAction bundle1 = MakeBundle1();
-			cache.SaveArtifactActionsAsync(new ArtifactAction[] { bundle1 }, cancellationToken).Wait();
+			await cache.SaveArtifactActionsAsync(new ArtifactAction[] { bundle1 }, cancellationToken);
 
-			ArtifactAction[] readBack1 = cache.QueryArtifactActionsAsync(new IoHash[] { bundle1.Key }, cancellationToken).Result;
+			ArtifactAction[] readBack1 = await cache.QueryArtifactActionsAsync(new IoHash[] { bundle1.Key }, cancellationToken);
 			Assert.AreEqual(1, readBack1.Length);
 
 			ArtifactAction bundle2 = MakeBundle2();
-			cache.SaveArtifactActionsAsync(new ArtifactAction[] { bundle2 }, cancellationToken).Wait();
+			await cache.SaveArtifactActionsAsync(new ArtifactAction[] { bundle2 }, cancellationToken);
 
-			ArtifactAction[] readBack2 = cache.QueryArtifactActionsAsync(new IoHash[] { bundle1.Key }, cancellationToken).Result;
+			ArtifactAction[] readBack2 = await cache.QueryArtifactActionsAsync(new IoHash[] { bundle1.Key }, cancellationToken);
 			Assert.AreEqual(2, readBack2.Length);
 
-			cache.FlushChangesAsync(cancellationToken).Wait();
+			await cache.FlushChangesAsync(cancellationToken);
 		}
 
 #if DISABLED
