Fix warnings being parsed as errors (due to "Warning" appearing capitalized).

diff --git a/Engine/Source/Programs/UnrealBuildTool.Tests/LogEventMatcherTests.cs b/Engine/Source/Programs/UnrealBuildTool.Tests/LogEventMatcherTests.cs
--- a/Engine/Source/Programs/UnrealBuildTool.Tests/LogEventMatcherTests.cs
+++ b/Engine/Source/Programs/UnrealBuildTool.Tests/LogEventMatcherTests.cs
@@ -1039,6 +1039,27 @@ namespace UnrealBuildToolTests
 			Assert.AreEqual(@"C:/Horde/Foo/Plugins/VerseAI/CompanionAI/Source/CompanionAI/Verse/CompanionAI.verse", fileProperty.Text);
 		}
 
+		[TestMethod]
+		public void VerseInEngineWarningMatcher()
+		{
+			string[] lines =
+			{
+				@"LogSolarisIde: Warning: C:/Horde/Plugins/PlayerProfileManager.verse(78,31, 84,14): Script Warning 2011: This expression can fail, but the meaning of failure in the right operand of 'set ... = ...' will change in a future version of Verse."
+			};
+
+			List<LogEvent> logEvents = Parse(String.Join("\n", lines));
+			Assert.AreEqual(1, logEvents.Count);
+			CheckEventGroup(logEvents.Slice(0, 1), 0, 1, LogLevel.Warning, KnownLogEvents.Compiler);
+
+			LogEvent logEvent = logEvents[0];
+			Assert.AreEqual("2011", logEvent.GetProperty("code").ToString());
+			Assert.AreEqual(LogLevel.Warning, logEvent.Level);
+
+			LogValue fileProperty = logEvents[0].GetProperty<LogValue>("file");
+			Assert.AreEqual(LogValueType.SourceFile, fileProperty.Type);
+			Assert.AreEqual(@"C:/Horde/Plugins/PlayerProfileManager.verse", fileProperty.Text);
+		}
+
 		static List<LogEvent> Parse(IEnumerable<string> lines)
 		{
 			return Parse(String.Join("\n", lines));
diff --git a/Engine/Source/Programs/UnrealBuildTool/Matchers/CompileEventMatcher.cs b/Engine/Source/Programs/UnrealBuildTool/Matchers/CompileEventMatcher.cs
--- a/Engine/Source/Programs/UnrealBuildTool/Matchers/CompileEventMatcher.cs
+++ b/Engine/Source/Programs/UnrealBuildTool/Matchers/CompileEventMatcher.cs
@@ -349,7 +349,7 @@ namespace UnrealBuildTool.Matchers
 		static LogLevel GetLogLevelFromSeverity(Match match)
 		{
 			string severity = match.Groups["severity"].Value;
-			if (severity.Equals("warning", StringComparison.Ordinal))
+			if (severity.Equals("warning", StringComparison.OrdinalIgnoreCase))
 			{
 				return LogLevel.Warning;
 			}
