UnrealBuildTool.Tests: Fix naming convention

#rnx

diff --git a/Engine/Source/Programs/UnrealBuildTool.Tests/PreprocessorExpressionTests.cs b/Engine/Source/Programs/UnrealBuildTool.Tests/PreprocessorExpressionTests.cs
--- a/Engine/Source/Programs/UnrealBuildTool.Tests/PreprocessorExpressionTests.cs
+++ b/Engine/Source/Programs/UnrealBuildTool.Tests/PreprocessorExpressionTests.cs
@@ -120,29 +120,29 @@ namespace UnrealBuildToolTests
 		/// <summary>
 		/// Tokenize the given expression and evaluate it, and check it matches the expected result
 		/// </summary>
-		/// <param name="Expression">The expression to evaluate, as a string</param>
-		/// <param name="ExpectedResult">The expected value of the expression</param>
-		static void RunTest(string Expression, long? ExpectedResult)
+		/// <param name="expression">The expression to evaluate, as a string</param>
+		/// <param name="expectedResult">The expected value of the expression</param>
+		static void RunTest(string expression, long? expectedResult)
 		{
-			using TokenReader Reader = new TokenReader(Expression);
+			using TokenReader reader = new TokenReader(expression);
 
-			List<Token> Tokens = new List<Token>();
-			while (Reader.MoveNext())
+			List<Token> tokens = new List<Token>();
+			while (reader.MoveNext())
 			{
-				Tokens.Add(Reader.Current);
+				tokens.Add(reader.Current);
 			}
 
-			long? Result;
+			long? result;
 			try
 			{
-				Result = PreprocessorExpression.Evaluate(new PreprocessorTestContext(), Tokens);
+				result = PreprocessorExpression.Evaluate(new PreprocessorTestContext(), tokens);
 			}
 			catch (PreprocessorException)
 			{
-				Result = null;
+				result = null;
 			}
 
-			Assert.AreEqual(ExpectedResult, Result);
+			Assert.AreEqual(expectedResult, result);
 		}
 	}
 }
diff --git a/Engine/Source/Programs/UnrealBuildTool.Tests/PreprocessorStateTests.cs b/Engine/Source/Programs/UnrealBuildTool.Tests/PreprocessorStateTests.cs
--- a/Engine/Source/Programs/UnrealBuildTool.Tests/PreprocessorStateTests.cs
+++ b/Engine/Source/Programs/UnrealBuildTool.Tests/PreprocessorStateTests.cs
@@ -14,133 +14,133 @@ namespace UnrealBuildToolTests
 		[TestMethod]
 		public void Run()
 		{
-			PreprocessorState BaseState = new PreprocessorState();
-			BaseState.DefineMacro(ParseMacro("FOO", null, "123"));
-			BaseState.DefineMacro(ParseMacro("BAR", new List<string> { "Arg1", "Arg2" }, "Arg1 + Arg2"));
-			Assert.AreEqual("m BAR(Arg1, Arg2)=Arg1 + Arg2\nm FOO=123\n", FormatState(BaseState));
+			PreprocessorState baseState = new PreprocessorState();
+			baseState.DefineMacro(ParseMacro("FOO", null, "123"));
+			baseState.DefineMacro(ParseMacro("BAR", new List<string> { "Arg1", "Arg2" }, "Arg1 + Arg2"));
+			Assert.AreEqual("m BAR(Arg1, Arg2)=Arg1 + Arg2\nm FOO=123\n", FormatState(baseState));
 
 			// Modify the state and revert it
 			{
-				PreprocessorState State = new PreprocessorState(BaseState);
-				Assert.AreEqual(FormatState(BaseState), FormatState(State));
-
-				PreprocessorTransform Transform1 = State.BeginCapture();
-				State.UndefMacro(Identifier.FindOrAdd("FOO"));
-				State.DefineMacro(ParseMacro("FOO2", null, "FOO()"));
-				State.PushBranch(PreprocessorBranch.Active);
-				Assert.AreEqual("b Active\nm BAR(Arg1, Arg2)=Arg1 + Arg2\nm FOO2=FOO()\n", FormatState(State));
-				Assert.AreEqual("+b Active\n+m FOO undef\n+m FOO2=FOO()\n", FormatTransform(Transform1));
-
-				PreprocessorState State2 = new PreprocessorState(BaseState);
-				State2.TryToApply(Transform1);
-				Assert.AreEqual(FormatState(State), FormatState(State2));
-
-				PreprocessorTransform Transform2 = State.BeginCapture();
-				State.PopBranch();
-				State.UndefMacro(Identifier.FindOrAdd("FOO"));
-				State.DefineMacro(ParseMacro("FOO", null, "123"));
-				State.UndefMacro(Identifier.FindOrAdd("FOO2"));
-				Assert.AreEqual(FormatState(BaseState), FormatState(State));
-				Assert.AreEqual("-b Active\n+m FOO=123\n+m FOO2 undef\n", FormatTransform(Transform2));
-
-				State2.TryToApply(Transform2);
-				Assert.AreEqual(FormatState(BaseState), FormatState(State2));
+				PreprocessorState state = new PreprocessorState(baseState);
+				Assert.AreEqual(FormatState(baseState), FormatState(state));
+
+				PreprocessorTransform transform1 = state.BeginCapture();
+				state.UndefMacro(Identifier.FindOrAdd("FOO"));
+				state.DefineMacro(ParseMacro("FOO2", null, "FOO()"));
+				state.PushBranch(PreprocessorBranch.Active);
+				Assert.AreEqual("b Active\nm BAR(Arg1, Arg2)=Arg1 + Arg2\nm FOO2=FOO()\n", FormatState(state));
+				Assert.AreEqual("+b Active\n+m FOO undef\n+m FOO2=FOO()\n", FormatTransform(transform1));
+
+				PreprocessorState state2 = new PreprocessorState(baseState);
+				state2.TryToApply(transform1);
+				Assert.AreEqual(FormatState(state), FormatState(state2));
+
+				PreprocessorTransform transform2 = state.BeginCapture();
+				state.PopBranch();
+				state.UndefMacro(Identifier.FindOrAdd("FOO"));
+				state.DefineMacro(ParseMacro("FOO", null, "123"));
+				state.UndefMacro(Identifier.FindOrAdd("FOO2"));
+				Assert.AreEqual(FormatState(baseState), FormatState(state));
+				Assert.AreEqual("-b Active\n+m FOO=123\n+m FOO2 undef\n", FormatTransform(transform2));
+
+				state2.TryToApply(transform2);
+				Assert.AreEqual(FormatState(baseState), FormatState(state2));
 			}
 
 			// Check the tracking of branches
 			{
-				PreprocessorState State = new PreprocessorState(BaseState);
-				PreprocessorTransform Transform1 = State.BeginCapture();
-				State.IsCurrentBranchActive();
-				Assert.AreEqual("=b Active\n", FormatTransform(Transform1));
-
-				State.PushBranch(PreprocessorBranch.Active);
-				PreprocessorTransform Transform2 = State.BeginCapture();
-				State.IsCurrentBranchActive();
-				State.PopBranch();
-				State.EndCapture();
-				Assert.AreEqual("-b Active\n", FormatTransform(Transform2));
-
-				State.PushBranch(0);
-				Assert.AreEqual("False", State.CanApply(Transform2).ToString());
-
-				State.PopBranch();
-				State.PushBranch(PreprocessorBranch.Active);
-				Assert.AreEqual("True", State.CanApply(Transform2).ToString());
+				PreprocessorState state = new PreprocessorState(baseState);
+				PreprocessorTransform transform1 = state.BeginCapture();
+				state.IsCurrentBranchActive();
+				Assert.AreEqual("=b Active\n", FormatTransform(transform1));
+
+				state.PushBranch(PreprocessorBranch.Active);
+				PreprocessorTransform transform2 = state.BeginCapture();
+				state.IsCurrentBranchActive();
+				state.PopBranch();
+				state.EndCapture();
+				Assert.AreEqual("-b Active\n", FormatTransform(transform2));
+
+				state.PushBranch(0);
+				Assert.AreEqual("False", state.CanApply(transform2).ToString());
+
+				state.PopBranch();
+				state.PushBranch(PreprocessorBranch.Active);
+				Assert.AreEqual("True", state.CanApply(transform2).ToString());
 			}
 		}
 
-		static string FormatState(PreprocessorState State)
+		static string FormatState(PreprocessorState state)
 		{
-			StringBuilder Result = new StringBuilder();
-			foreach (PreprocessorBranch Branch in State.CurrentBranches)
+			StringBuilder result = new StringBuilder();
+			foreach (PreprocessorBranch branch in state.CurrentBranches)
 			{
-				Result.AppendFormat("b {0}\n", Branch);
+				result.AppendFormat("b {0}\n", branch);
 			}
-			foreach (PreprocessorMacro Macro in State.CurrentMacros.OrderBy(x => x.Name))
+			foreach (PreprocessorMacro macro in state.CurrentMacros.OrderBy(x => x.Name))
 			{
-				Result.AppendFormat("m {0}\n", Macro.ToString().TrimEnd());
+				result.AppendFormat("m {0}\n", macro.ToString().TrimEnd());
 			}
-			return Result.ToString();
+			return result.ToString();
 		}
 
-		static string FormatTransform(PreprocessorTransform Transform)
+		static string FormatTransform(PreprocessorTransform transform)
 		{
-			StringBuilder Result = new StringBuilder();
-			if (Transform.RequireTopmostActive.HasValue)
+			StringBuilder result = new StringBuilder();
+			if (transform.RequireTopmostActive.HasValue)
 			{
-				Result.AppendFormat("=b {0}\n", Transform.RequireTopmostActive.Value ? "Active" : "0");
+				result.AppendFormat("=b {0}\n", transform.RequireTopmostActive.Value ? "Active" : "0");
 			}
-			foreach (PreprocessorBranch Branch in Transform.RequiredBranches)
+			foreach (PreprocessorBranch branch in transform.RequiredBranches)
 			{
-				Result.AppendFormat("-b {0}\n", Branch);
+				result.AppendFormat("-b {0}\n", branch);
 			}
-			foreach (PreprocessorBranch Branch in Transform.NewBranches)
+			foreach (PreprocessorBranch branch in transform.NewBranches)
 			{
-				Result.AppendFormat("+b {0}\n", Branch);
+				result.AppendFormat("+b {0}\n", branch);
 			}
-			foreach (KeyValuePair<Identifier, PreprocessorMacro> Macro in Transform.RequiredMacros)
+			foreach (KeyValuePair<Identifier, PreprocessorMacro> macro in transform.RequiredMacros)
 			{
-				if (Macro.Value == null)
+				if (macro.Value == null)
 				{
-					Result.AppendFormat("=m {0} undef\n", Macro.Key);
+					result.AppendFormat("=m {0} undef\n", macro.Key);
 				}
 				else
 				{
-					Result.AppendFormat("=m {0}\n", Macro.Value.ToString().TrimEnd());
+					result.AppendFormat("=m {0}\n", macro.Value.ToString().TrimEnd());
 				}
 			}
-			foreach (KeyValuePair<Identifier, PreprocessorMacro> Macro in Transform.NewMacros)
+			foreach (KeyValuePair<Identifier, PreprocessorMacro> macro in transform.NewMacros)
 			{
-				if (Macro.Value == null)
+				if (macro.Value == null)
 				{
-					Result.AppendFormat("+m {0} undef\n", Macro.Key);
+					result.AppendFormat("+m {0} undef\n", macro.Key);
 				}
 				else
 				{
-					Result.AppendFormat("+m {0}\n", Macro.Value.ToString().TrimEnd());
+					result.AppendFormat("+m {0}\n", macro.Value.ToString().TrimEnd());
 				}
 			}
-			return Result.ToString();
+			return result.ToString();
 		}
 
-		static PreprocessorMacro ParseMacro(string Name, List<string> Parameters, string Value)
+		static PreprocessorMacro ParseMacro(string name, List<string> parameters, string value)
 		{
-			List<Token> Tokens = new List<Token>();
+			List<Token> tokens = new List<Token>();
 
-			using TokenReader Reader = new TokenReader(Value);
-			while (Reader.MoveNext())
+			using TokenReader reader = new TokenReader(value);
+			while (reader.MoveNext())
 			{
-				Tokens.Add(Reader.Current);
+				tokens.Add(reader.Current);
 			}
 
-			if (Parameters == null)
+			if (parameters == null)
 			{
-				return new PreprocessorMacro(Identifier.FindOrAdd(Name), null, Tokens);
+				return new PreprocessorMacro(Identifier.FindOrAdd(name), null, tokens);
 			}
 			else
 			{
-				return new PreprocessorMacro(Identifier.FindOrAdd(Name), Parameters.Select(x => Identifier.FindOrAdd(x)).ToList(), Tokens);
+				return new PreprocessorMacro(Identifier.FindOrAdd(name), parameters.Select(x => Identifier.FindOrAdd(x)).ToList(), tokens);
 			}
 		}
 	}
diff --git a/Engine/Source/Programs/UnrealBuildTool.Tests/PreprocessorTests.cs b/Engine/Source/Programs/UnrealBuildTool.Tests/PreprocessorTests.cs
--- a/Engine/Source/Programs/UnrealBuildTool.Tests/PreprocessorTests.cs
+++ b/Engine/Source/Programs/UnrealBuildTool.Tests/PreprocessorTests.cs
@@ -218,58 +218,58 @@ namespace UnrealBuildToolTests
 		/// <summary>
 		/// Preprocess a fragment of code, and check it results in an expected sequence of tokens
 		/// </summary>
-		/// <param name="Fragment">The code fragment to preprocess</param>
-		/// <param name="ExpectedResult">The expected sequence of tokens, as a string. Null to indicate that the input is invalid, and an exception is expected.</param>
-		static void RunTest(string Fragment, string ExpectedResult)
+		/// <param name="fragment">The code fragment to preprocess</param>
+		/// <param name="expectedResult">The expected sequence of tokens, as a string. Null to indicate that the input is invalid, and an exception is expected.</param>
+		static void RunTest(string fragment, string expectedResult)
 		{
-			string[] Lines = Fragment.Split('\n');
+			string[] lines = fragment.Split('\n');
 
-			SourceFile File = new SourceFile(null, TokenReader.GetNullTerminatedByteArray(Fragment));
+			SourceFile file = new SourceFile(null, TokenReader.GetNullTerminatedByteArray(fragment));
 
-			Preprocessor Instance = new Preprocessor();
+			Preprocessor instance = new Preprocessor();
 
-			string Result;
+			string result;
 			try
 			{
-				List<Token> OutputTokens = new List<Token>();
-				for (int MarkupIdx = 0; MarkupIdx < File.Markup.Length; MarkupIdx++)
+				List<Token> outputTokens = new List<Token>();
+				for (int markupIdx = 0; markupIdx < file.Markup.Length; markupIdx++)
 				{
-					SourceFileMarkup Markup = File.Markup[MarkupIdx];
-					if (Markup.Type == SourceFileMarkupType.Text)
+					SourceFileMarkup markup = file.Markup[markupIdx];
+					if (markup.Type == SourceFileMarkupType.Text)
 					{
-						if (Instance.IsCurrentBranchActive())
+						if (instance.IsCurrentBranchActive())
 						{
-							StringBuilder SourceText = new StringBuilder();
+							StringBuilder sourceText = new StringBuilder();
 
-							int LastLineIndex = (MarkupIdx + 1 < File.Markup.Length) ? File.Markup[MarkupIdx + 1].LineNumber - 1 : Lines.Length;
-							for (int LineIndex = Markup.LineNumber - 1; LineIndex < LastLineIndex; LineIndex++)
+							int lastLineIndex = (markupIdx + 1 < file.Markup.Length) ? file.Markup[markupIdx + 1].LineNumber - 1 : lines.Length;
+							for (int lineIndex = markup.LineNumber - 1; lineIndex < lastLineIndex; lineIndex++)
 							{
-								SourceText.AppendLine(Lines[LineIndex]);
+								sourceText.AppendLine(lines[lineIndex]);
 							}
 
-							List<Token> Tokens = new List<Token>();
+							List<Token> tokens = new List<Token>();
 
-							using TokenReader Reader = new TokenReader(SourceText.ToString());
-							while (Reader.MoveNext())
+							using TokenReader reader = new TokenReader(sourceText.ToString());
+							while (reader.MoveNext())
 							{
-								Tokens.Add(Reader.Current);
+								tokens.Add(reader.Current);
 							}
 
-							Instance.ExpandMacros(Tokens, OutputTokens, false, null);
+							instance.ExpandMacros(tokens, outputTokens, false, null);
 						}
 					}
 					else
 					{
-						Instance.ParseMarkup(Markup.Type, Markup.Tokens, null);
+						instance.ParseMarkup(markup.Type, markup.Tokens, null);
 					}
 				}
-				Result = Token.Format(OutputTokens);
+				result = Token.Format(outputTokens);
 			}
 			catch (PreprocessorException)
 			{
-				Result = null;
+				result = null;
 			}
-			Assert.AreEqual(ExpectedResult, Result);
+			Assert.AreEqual(expectedResult, result);
 		}
 	}
 }
diff --git a/Engine/Source/Programs/UnrealBuildTool.Tests/SourceFileTests.cs b/Engine/Source/Programs/UnrealBuildTool.Tests/SourceFileTests.cs
--- a/Engine/Source/Programs/UnrealBuildTool.Tests/SourceFileTests.cs
+++ b/Engine/Source/Programs/UnrealBuildTool.Tests/SourceFileTests.cs
@@ -20,65 +20,65 @@ namespace UnrealBuildToolTests
 		[TestMethod]
 		public void Run()
 		{
-			List<DirectoryReference> BaseDirectories = new List<DirectoryReference>();
-			BaseDirectories.Add(DirectoryReference.Combine(Unreal.EngineSourceDirectory, "Runtime"));
-			BaseDirectories.Add(DirectoryReference.Combine(Unreal.EngineSourceDirectory, "Developer"));
-			BaseDirectories.Add(DirectoryReference.Combine(Unreal.EngineSourceDirectory, "Editor"));
+			List<DirectoryReference> baseDirectories = new List<DirectoryReference>();
+			baseDirectories.Add(DirectoryReference.Combine(Unreal.EngineSourceDirectory, "Runtime"));
+			baseDirectories.Add(DirectoryReference.Combine(Unreal.EngineSourceDirectory, "Developer"));
+			baseDirectories.Add(DirectoryReference.Combine(Unreal.EngineSourceDirectory, "Editor"));
 
-			foreach (FileReference PluginFile in PluginsBase.EnumeratePlugins((FileReference)null))
+			foreach (FileReference pluginFile in PluginsBase.EnumeratePlugins((FileReference)null))
 			{
-				DirectoryReference PluginSourceDir = DirectoryReference.Combine(PluginFile.Directory, "Source");
-				if (DirectoryReference.Exists(PluginSourceDir))
+				DirectoryReference pluginSourceDir = DirectoryReference.Combine(pluginFile.Directory, "Source");
+				if (DirectoryReference.Exists(pluginSourceDir))
 				{
-					BaseDirectories.Add(PluginSourceDir);
+					baseDirectories.Add(pluginSourceDir);
 				}
 			}
 
-			ConcurrentBag<SourceFile> SourceFiles = new ConcurrentBag<SourceFile>();
+			ConcurrentBag<SourceFile> sourceFiles = new ConcurrentBag<SourceFile>();
 			using (GlobalTracer.Instance.BuildSpan("Scanning source files").StartActive())
 			{
-				using (ThreadPoolWorkQueue Queue = new ThreadPoolWorkQueue())
+				using (ThreadPoolWorkQueue queue = new ThreadPoolWorkQueue())
 				{
-					foreach (DirectoryReference BaseDirectory in BaseDirectories)
+					foreach (DirectoryReference baseDirectory in baseDirectories)
 					{
-						Queue.Enqueue(() => ParseSourceFiles(DirectoryItem.GetItemByDirectoryReference(BaseDirectory), SourceFiles, Queue));
+						queue.Enqueue(() => ParseSourceFiles(DirectoryItem.GetItemByDirectoryReference(baseDirectory), sourceFiles, queue));
 					}
 				}
 			}
 
-			FileReference TempDataFile = FileReference.Combine(Unreal.EngineDirectory, "Intermediate", "Temp", "SourceFileTests.bin");
-			DirectoryReference.CreateDirectory(TempDataFile.Directory);
+			FileReference tempDataFile = FileReference.Combine(Unreal.EngineDirectory, "Intermediate", "Temp", "SourceFileTests.bin");
+			DirectoryReference.CreateDirectory(tempDataFile.Directory);
 
 			using (GlobalTracer.Instance.BuildSpan("Writing source file data").StartActive())
 			{
-				using (BinaryArchiveWriter Writer = new BinaryArchiveWriter(TempDataFile))
+				using (BinaryArchiveWriter writer = new BinaryArchiveWriter(tempDataFile))
 				{
-					Writer.WriteList(SourceFiles.ToList(), x => x.Write(Writer));
+					writer.WriteList(sourceFiles.ToList(), x => x.Write(writer));
 				}
 			}
 
-			List<SourceFile> ReadSourceFiles = new List<SourceFile>();
+			List<SourceFile> readSourceFiles = new List<SourceFile>();
 			using (GlobalTracer.Instance.BuildSpan("Reading source file data").StartActive())
 			{
-				using (BinaryArchiveReader Reader = new BinaryArchiveReader(TempDataFile))
+				using (BinaryArchiveReader reader = new BinaryArchiveReader(tempDataFile))
 				{
-					ReadSourceFiles = Reader.ReadList(() => new SourceFile(Reader));
+					readSourceFiles = reader.ReadList(() => new SourceFile(reader));
 				}
 			}
 		}
 
-		static void ParseSourceFiles(DirectoryItem Directory, ConcurrentBag<SourceFile> SourceFiles, ThreadPoolWorkQueue Queue)
+		static void ParseSourceFiles(DirectoryItem directory, ConcurrentBag<SourceFile> sourceFiles, ThreadPoolWorkQueue queue)
 		{
-			foreach (DirectoryItem SubDirectory in Directory.EnumerateDirectories())
+			foreach (DirectoryItem subDirectory in directory.EnumerateDirectories())
 			{
-				Queue.Enqueue(() => ParseSourceFiles(SubDirectory, SourceFiles, Queue));
+				queue.Enqueue(() => ParseSourceFiles(subDirectory, sourceFiles, queue));
 			}
 
-			foreach (FileItem File in Directory.EnumerateFiles())
+			foreach (FileItem file in directory.EnumerateFiles())
 			{
-				if (File.HasExtension(".h") || File.HasExtension(".cpp"))
+				if (file.HasExtension(".h") || file.HasExtension(".cpp"))
 				{
-					Queue.Enqueue(() => SourceFiles.Add(new SourceFile(File)));
+					queue.Enqueue(() => sourceFiles.Add(new SourceFile(file)));
 				}
 			}
 		}
diff --git a/Engine/Source/Programs/UnrealBuildTool.Tests/TokenReaderTests.cs b/Engine/Source/Programs/UnrealBuildTool.Tests/TokenReaderTests.cs
--- a/Engine/Source/Programs/UnrealBuildTool.Tests/TokenReaderTests.cs
+++ b/Engine/Source/Programs/UnrealBuildTool.Tests/TokenReaderTests.cs
@@ -132,42 +132,42 @@ namespace UnrealBuildToolTests
 			RunLineNumberTest("first line \\\n+", 2);
 		}
 
-		public static void RunLexerTest(string InputText, string ExpectedOutputText)
+		public static void RunLexerTest(string inputText, string expectedOutputText)
 		{
-			List<Token> Tokens = new List<Token>();
+			List<Token> tokens = new List<Token>();
 
-			using TokenReader Reader = new TokenReader(InputText);
-			while (Reader.MoveNext())
+			using TokenReader reader = new TokenReader(inputText);
+			while (reader.MoveNext())
 			{
-				Tokens.Add(Reader.Current);
+				tokens.Add(reader.Current);
 			}
 
-			string OutputText = String.Join(", ", Tokens.Select(x => FormatToken(x)));
-			Assert.AreEqual(ExpectedOutputText, OutputText);
+			string outputText = String.Join(", ", tokens.Select(x => FormatToken(x)));
+			Assert.AreEqual(expectedOutputText, outputText);
 		}
 
-		static string FormatToken(Token Token)
+		static string FormatToken(Token token)
 		{
-			StringBuilder Result = new StringBuilder();
-			Result.AppendFormat(Token.Type.ToString());
-			if (Token.Type == TokenType.Identifier || Token.Type == TokenType.Number || Token.Type == TokenType.Character || Token.Type == TokenType.String)
+			StringBuilder result = new StringBuilder();
+			result.AppendFormat(token.Type.ToString());
+			if (token.Type == TokenType.Identifier || token.Type == TokenType.Number || token.Type == TokenType.Character || token.Type == TokenType.String)
 			{
-				Result.AppendFormat("({0})", Token.Text);
+				result.AppendFormat("({0})", token.Text);
 			}
-			return Result.ToString();
+			return result.ToString();
 		}
 
-		public static void RunLineNumberTest(string InputText, int ExpectedLineNumber)
+		public static void RunLineNumberTest(string inputText, int expectedLineNumber)
 		{
-			using TokenReader Reader = new TokenReader(InputText);
-			while (Reader.MoveNext())
+			using TokenReader reader = new TokenReader(inputText);
+			while (reader.MoveNext())
 			{
-				if (Reader.Current.Type == TokenType.Plus)
+				if (reader.Current.Type == TokenType.Plus)
 				{
 					break;
 				}
 			}
-			Assert.AreEqual(ExpectedLineNumber, Reader.LineNumber);
+			Assert.AreEqual(expectedLineNumber, reader.LineNumber);
 		}
 	}
 }
