UnrealBuildTool.Tests: Fix nuget packages

#rnx

diff --git a/Engine/Source/Programs/UnrealBuildTool.Tests/UnrealBuildTool.Tests.csproj b/Engine/Source/Programs/UnrealBuildTool.Tests/UnrealBuildTool.Tests.csproj
--- a/Engine/Source/Programs/UnrealBuildTool.Tests/UnrealBuildTool.Tests.csproj
+++ b/Engine/Source/Programs/UnrealBuildTool.Tests/UnrealBuildTool.Tests.csproj
@@ -8,10 +8,13 @@
   </PropertyGroup>
 
   <ItemGroup>
-    <PackageReference Include="Microsoft.NET.Test.Sdk" Version="16.5.0" />
-    <PackageReference Include="MSTest.TestAdapter" Version="2.1.0" />
-    <PackageReference Include="MSTest.TestFramework" Version="2.1.0" />
-    <PackageReference Include="coverlet.collector" Version="1.2.0" />
+    <PackageReference Include="Microsoft.NET.Test.Sdk" Version="17.3.0" />
+    <PackageReference Include="MSTest.TestAdapter" Version="2.2.10" />
+    <PackageReference Include="MSTest.TestFramework" Version="2.2.10" />
+    <PackageReference Include="coverlet.collector" Version="3.1.2">
+      <PrivateAssets>all</PrivateAssets>
+      <IncludeAssets>runtime; build; native; contentfiles; analyzers; buildtransitive</IncludeAssets>
+    </PackageReference>
   </ItemGroup>
 
   <ItemGroup>
