UnrealBuildTool: Fix tests

#rnx

diff --git a/Engine/Source/Programs/UnrealBuildTool.Tests/ArtifactsTest.cs b/Engine/Source/Programs/UnrealBuildTool.Tests/ArtifactsTest.cs
--- a/Engine/Source/Programs/UnrealBuildTool.Tests/ArtifactsTest.cs
+++ b/Engine/Source/Programs/UnrealBuildTool.Tests/ArtifactsTest.cs
@@ -76,23 +76,7 @@ namespace UnrealBuildToolTests
 			ArtifactAction[] readBack2 = await cache.QueryArtifactActionsAsync(new IoHash[] { bundle1.Key }, cancellationToken);
 			Assert.AreEqual(2, readBack2.Length);
 
-			await cache.FlushChangesAsync(cancellationToken);
+			//await cache.FlushChangesAsync(cancellationToken);
 		}
-
-#if DISABLED
-		[TestMethod]
-		public void ArtifactBundleStorageTest2()
-		{
-			CancellationToken cancellationToken = default;
-
-			IArtifactCache cache = HordeStorageArtifactCache.CreateMemoryCache(NullLogger.Instance);
-
-			cache.WaitForReadyAsync().Wait(cancellationToken);
-
-			ArtifactBundle bundle1 = MakeBundle1();
-			ArtifactBundle[] readBack2 = cache.QueryArtifactBundles(new IoHash[] { bundle1.Key }, cancellationToken);
-			Assert.AreEqual(2, readBack2.Length);
-		}
-#endif
 	}
 }
diff --git a/Engine/Source/Programs/UnrealBuildTool.Tests/LogEventMatcherTests.cs b/Engine/Source/Programs/UnrealBuildTool.Tests/LogEventMatcherTests.cs
--- a/Engine/Source/Programs/UnrealBuildTool.Tests/LogEventMatcherTests.cs
+++ b/Engine/Source/Programs/UnrealBuildTool.Tests/LogEventMatcherTests.cs
@@ -907,8 +907,9 @@ namespace UnrealBuildToolTests
 			};
 
 			List<LogEvent> logEvents = Parse(lines);
-			Assert.AreEqual(1, logEvents.Count);
-			CheckEventGroup(logEvents, 5, 1, LogLevel.Information, KnownLogEvents.Systemic_Xge_CacheLimit);
+			Assert.AreEqual(3, logEvents.Count);
+			CheckEventGroup(logEvents.Slice(0, 1), 5, 1, LogLevel.Information, KnownLogEvents.Systemic_Xge_CacheLimit);
+			CheckEventGroup(logEvents.Slice(1, 2), 7, 2, LogLevel.Information, KnownLogEvents.Systemic_Xge);
 		}
 
 		[TestMethod]
@@ -953,9 +954,10 @@ namespace UnrealBuildToolTests
 			};
 
 			List<LogEvent> logEvents = Parse(lines);
-			Assert.AreEqual(2, logEvents.Count);
+			Assert.AreEqual(4, logEvents.Count);
 			CheckEventGroup(logEvents.Slice(0, 1), 5, 1, LogLevel.Information, KnownLogEvents.Systemic_Xge_CacheLimit);
 			CheckEventGroup(logEvents.Slice(1, 1), 6, 1, LogLevel.Information, KnownLogEvents.Systemic_Xge_CacheLimit);
+			CheckEventGroup(logEvents.Slice(2, 2), 7, 2, LogLevel.Information, KnownLogEvents.Systemic_Xge);
 		}
 
 		[TestMethod]
diff --git a/Engine/Source/Programs/UnrealBuildTool.Tests/PluginDescriptorTests.cs b/Engine/Source/Programs/UnrealBuildTool.Tests/PluginDescriptorTests.cs
--- a/Engine/Source/Programs/UnrealBuildTool.Tests/PluginDescriptorTests.cs
+++ b/Engine/Source/Programs/UnrealBuildTool.Tests/PluginDescriptorTests.cs
@@ -25,7 +25,7 @@ namespace UnrealBuildToolTests
 		}
 
 		[ClassInitialize]
-		public static void Setup()
+		public static void Setup(TestContext _)
 		{
 			PluginDescriptorTests.s_tempDirectory = CreateTempDir();
 		}
@@ -256,7 +256,7 @@ namespace UnrealBuildToolTests
 		}
 
 		private static string FixVerbatimStringIdentation(string verbatimString)
-	{
+		{
 			if (String.IsNullOrEmpty(verbatimString))
 			{
 				return verbatimString;
@@ -275,12 +275,12 @@ namespace UnrealBuildToolTests
 				// Now we're on the first line with the { we're just going to assume that this is the first line 
 				if (!bFoundTabLevel)
 				{
-					numberOfTabs= line.TakeWhile(x => x == '\t').Count();
+					numberOfTabs = line.TakeWhile(x => x == '\t').Count();
 					bFoundTabLevel = true;
 				}
 				jsonStringBuilder.AppendLine(line.Substring(numberOfTabs));
 			}
 			return jsonStringBuilder.ToString();
-	}
+		}
 	}
 }
