Add Analyze config to C# projects

#rnx

diff --git a/Engine/Source/Programs/Shared/EpicGames.IoHash.Tests/EpicGames.IoHash.Tests.csproj b/Engine/Source/Programs/Shared/EpicGames.IoHash.Tests/EpicGames.IoHash.Tests.csproj
--- a/Engine/Source/Programs/Shared/EpicGames.IoHash.Tests/EpicGames.IoHash.Tests.csproj
+++ b/Engine/Source/Programs/Shared/EpicGames.IoHash.Tests/EpicGames.IoHash.Tests.csproj
@@ -3,9 +3,17 @@
   <PropertyGroup>
     <TargetFramework>net6.0</TargetFramework>
 	<TargetFramework Condition="'$(UE_DOTNET_VERSION)' != ''">$(UE_DOTNET_VERSION)</TargetFramework>
-    <Configurations>Debug;Release;Development</Configurations>
+    <Configurations>Debug;Release;Development;Analyze</Configurations>
     <Nullable>enable</Nullable>
     <IsPackable>false</IsPackable>
+    <NoWarn>1591</NoWarn>
+  </PropertyGroup>
+
+  <PropertyGroup Condition=" '$(Configuration)|$(Platform)' == 'Analyze|AnyCPU' ">
+    <RunAnalyzersDuringBuild>True</RunAnalyzersDuringBuild>
+    <TreatWarningsAsErrors>false</TreatWarningsAsErrors>
+    <GenerateDocumentationFile>True</GenerateDocumentationFile>
+    <DocumentationFile></DocumentationFile>
   </PropertyGroup>
 
   <ItemGroup>
diff --git a/Engine/Source/Programs/Shared/EpicGames.Jupiter/EpicGames.Jupiter.csproj b/Engine/Source/Programs/Shared/EpicGames.Jupiter/EpicGames.Jupiter.csproj
--- a/Engine/Source/Programs/Shared/EpicGames.Jupiter/EpicGames.Jupiter.csproj
+++ b/Engine/Source/Programs/Shared/EpicGames.Jupiter/EpicGames.Jupiter.csproj
@@ -7,9 +7,10 @@
     <RootNamespace>EpicGames.Jupiter</RootNamespace>
     <GenerateAssemblyInfo>false</GenerateAssemblyInfo>
     <GenerateTargetFrameworkAttribute>false</GenerateTargetFrameworkAttribute>
-    <Configurations>Debug;Release;Development</Configurations>
+    <Configurations>Debug;Release;Development;Analyze</Configurations>
     <DebugType>pdbonly</DebugType>
     <DebugType Condition="'$([System.Runtime.InteropServices.RuntimeInformation]::IsOSPlatform($([System.Runtime.InteropServices.OSPlatform]::Windows)))' == 'true' And '$([System.Runtime.InteropServices.RuntimeInformation]::OSArchitecture)' == 'Arm64'">portable</DebugType>
+    <NoWarn>1591</NoWarn>
   </PropertyGroup>
   <PropertyGroup Condition=" '$(Configuration)|$(Platform)' == 'Debug|AnyCPU' ">
     <DocumentationFile></DocumentationFile>
@@ -22,6 +23,12 @@
     <DocumentationFile></DocumentationFile>
     <Optimize>true</Optimize>
   </PropertyGroup>
+  <PropertyGroup Condition=" '$(Configuration)|$(Platform)' == 'Analyze|AnyCPU' ">
+    <RunAnalyzersDuringBuild>True</RunAnalyzersDuringBuild>
+    <TreatWarningsAsErrors>false</TreatWarningsAsErrors>
+    <GenerateDocumentationFile>True</GenerateDocumentationFile>
+    <DocumentationFile></DocumentationFile>
+  </PropertyGroup>
   <ItemGroup>
     <Compile Include="..\..\Shared\MetaData.cs">
       <Link>Properties\MetaData.cs</Link>
diff --git a/Engine/Source/Programs/Shared/EpicGames.MongoDB/EpicGames.MongoDB.csproj b/Engine/Source/Programs/Shared/EpicGames.MongoDB/EpicGames.MongoDB.csproj
--- a/Engine/Source/Programs/Shared/EpicGames.MongoDB/EpicGames.MongoDB.csproj
+++ b/Engine/Source/Programs/Shared/EpicGames.MongoDB/EpicGames.MongoDB.csproj
@@ -6,7 +6,7 @@
     <RootNamespace>EpicGames.Core</RootNamespace>
     <GenerateAssemblyInfo>false</GenerateAssemblyInfo>
     <GenerateTargetFrameworkAttribute>false</GenerateTargetFrameworkAttribute>
-    <Configurations>Debug;Release;Development</Configurations>
+    <Configurations>Debug;Release;Development;Analyze</Configurations>
     <Nullable>enable</Nullable>
   </PropertyGroup>
   <PropertyGroup Condition=" '$(Configuration)|$(Platform)' == 'Debug|AnyCPU' ">
@@ -20,6 +20,12 @@
     <DocumentationFile></DocumentationFile>
     <Optimize>true</Optimize>
   </PropertyGroup>
+  <PropertyGroup Condition=" '$(Configuration)|$(Platform)' == 'Analyze|AnyCPU' ">
+    <RunAnalyzersDuringBuild>True</RunAnalyzersDuringBuild>
+    <TreatWarningsAsErrors>false</TreatWarningsAsErrors>
+    <GenerateDocumentationFile>True</GenerateDocumentationFile>
+    <DocumentationFile></DocumentationFile>
+  </PropertyGroup>
   <ItemGroup>
     <Compile Include="..\..\Shared\MetaData.cs">
       <Link>Properties\MetaData.cs</Link>
diff --git a/Engine/Source/Programs/Shared/EpicGames.MongoDB/SingletonDocument.cs b/Engine/Source/Programs/Shared/EpicGames.MongoDB/SingletonDocument.cs
--- a/Engine/Source/Programs/Shared/EpicGames.MongoDB/SingletonDocument.cs
+++ b/Engine/Source/Programs/Shared/EpicGames.MongoDB/SingletonDocument.cs
@@ -49,7 +49,6 @@ namespace EpicGames.MongoDB
 		/// <summary>
 		/// Private constructor for serialization
 		/// </summary>
-		/// <param name="Unused"></param>
 		[BsonConstructor]
 		protected SingletonBase()
 		{
@@ -136,7 +135,7 @@ namespace EpicGames.MongoDB
 		/// <summary>
 		/// Constructor
 		/// </summary>
-		/// <param name="DatabaseService">The database service instance</param>
+		/// <param name="collection">The database service instance</param>
 		public SingletonDocument(IMongoCollection<SingletonBase> collection)
 		{
 			_collection = collection.OfType<T>();
diff --git a/Engine/Source/Programs/Shared/EpicGames.MsBuild/EpicGames.MsBuild.csproj b/Engine/Source/Programs/Shared/EpicGames.MsBuild/EpicGames.MsBuild.csproj
--- a/Engine/Source/Programs/Shared/EpicGames.MsBuild/EpicGames.MsBuild.csproj
+++ b/Engine/Source/Programs/Shared/EpicGames.MsBuild/EpicGames.MsBuild.csproj
@@ -8,7 +8,7 @@
     <RootNamespace>EpicGames.MsBuild</RootNamespace>
     <GenerateAssemblyInfo>false</GenerateAssemblyInfo>
     <GenerateTargetFrameworkAttribute>false</GenerateTargetFrameworkAttribute>
-    <Configurations>Debug;Release;Development</Configurations>
+    <Configurations>Debug;Release;Development;Analyze</Configurations>
     <DebugType>pdbonly</DebugType>
     <DebugType Condition="'$([System.Runtime.InteropServices.RuntimeInformation]::IsOSPlatform($([System.Runtime.InteropServices.OSPlatform]::Windows)))' == 'true' And '$([System.Runtime.InteropServices.RuntimeInformation]::OSArchitecture)' == 'Arm64'">portable</DebugType>
     <ProduceReferenceAssemblyInOutDir>true</ProduceReferenceAssemblyInOutDir>
@@ -24,6 +24,12 @@
     <DocumentationFile></DocumentationFile>
     <Optimize>true</Optimize>
   </PropertyGroup>
+  <PropertyGroup Condition=" '$(Configuration)|$(Platform)' == 'Analyze|AnyCPU' ">
+    <RunAnalyzersDuringBuild>True</RunAnalyzersDuringBuild>
+    <TreatWarningsAsErrors>false</TreatWarningsAsErrors>
+    <GenerateDocumentationFile>True</GenerateDocumentationFile>
+    <DocumentationFile></DocumentationFile>
+  </PropertyGroup>
   
   <ItemGroup>
     <Compile Include="..\MetaData.cs">
diff --git a/Engine/Source/Programs/Shared/EpicGames.Tracing.Tests/EpicGames.Tracing.Tests.csproj b/Engine/Source/Programs/Shared/EpicGames.Tracing.Tests/EpicGames.Tracing.Tests.csproj
--- a/Engine/Source/Programs/Shared/EpicGames.Tracing.Tests/EpicGames.Tracing.Tests.csproj
+++ b/Engine/Source/Programs/Shared/EpicGames.Tracing.Tests/EpicGames.Tracing.Tests.csproj
@@ -10,9 +10,16 @@
     <GenerateAssemblyInformationalVersionAttribute>false</GenerateAssemblyInformationalVersionAttribute>
     <GenerateAssemblyProductAttribute>false</GenerateAssemblyProductAttribute>
     <GenerateAssemblyVersionAttribute>false</GenerateAssemblyVersionAttribute>
-    <Configurations>Debug;Release;Development</Configurations>
+    <Configurations>Debug;Release;Development;Analyze</Configurations>
     <IsPackable>false</IsPackable>
     <Nullable>enable</Nullable>
+    <NoWarn>1591</NoWarn>
+  </PropertyGroup>
+  <PropertyGroup Condition=" '$(Configuration)|$(Platform)' == 'Analyze|AnyCPU' ">
+    <RunAnalyzersDuringBuild>True</RunAnalyzersDuringBuild>
+    <TreatWarningsAsErrors>false</TreatWarningsAsErrors>
+    <GenerateDocumentationFile>True</GenerateDocumentationFile>
+    <DocumentationFile></DocumentationFile>
   </PropertyGroup>
   <ItemGroup>
     <Compile Include="..\..\Shared\MetaData.cs">
diff --git a/Engine/Source/Programs/Shared/EpicGames.Tracing/EpicGames.Tracing.csproj b/Engine/Source/Programs/Shared/EpicGames.Tracing/EpicGames.Tracing.csproj
--- a/Engine/Source/Programs/Shared/EpicGames.Tracing/EpicGames.Tracing.csproj
+++ b/Engine/Source/Programs/Shared/EpicGames.Tracing/EpicGames.Tracing.csproj
@@ -7,8 +7,9 @@
     <RootNamespace>EpicGames.Tracing</RootNamespace>
     <GenerateAssemblyInfo>false</GenerateAssemblyInfo>
     <GenerateTargetFrameworkAttribute>false</GenerateTargetFrameworkAttribute>
-    <Configurations>Debug;Release;Development</Configurations>
+    <Configurations>Debug;Release;Development;Analyze</Configurations>
     <Nullable>enable</Nullable>
+    <NoWarn>1591</NoWarn>
   </PropertyGroup>
   <PropertyGroup Condition=" '$(Configuration)|$(Platform)' == 'Debug|AnyCPU' ">
     <DocumentationFile></DocumentationFile>
@@ -21,6 +22,12 @@
     <DocumentationFile></DocumentationFile>
     <Optimize>true</Optimize>
   </PropertyGroup>
+  <PropertyGroup Condition=" '$(Configuration)|$(Platform)' == 'Analyze|AnyCPU' ">
+    <RunAnalyzersDuringBuild>True</RunAnalyzersDuringBuild>
+    <TreatWarningsAsErrors>false</TreatWarningsAsErrors>
+    <GenerateDocumentationFile>True</GenerateDocumentationFile>
+    <DocumentationFile></DocumentationFile>
+  </PropertyGroup>
   <ItemGroup>
     <Compile Include="..\..\Shared\MetaData.cs">
       <Link>Properties\MetaData.cs</Link>
diff --git a/Engine/Source/Programs/Shared/EpicGames.UBA/EpicGames.UBA.csproj b/Engine/Source/Programs/Shared/EpicGames.UBA/EpicGames.UBA.csproj
--- a/Engine/Source/Programs/Shared/EpicGames.UBA/EpicGames.UBA.csproj
+++ b/Engine/Source/Programs/Shared/EpicGames.UBA/EpicGames.UBA.csproj
@@ -8,7 +8,7 @@
     <RootNamespace>EpicGames.UBA</RootNamespace>
     <GenerateAssemblyInfo>false</GenerateAssemblyInfo>
     <GenerateTargetFrameworkAttribute>false</GenerateTargetFrameworkAttribute>
-    <Configurations>Debug;Release;Development</Configurations>
+    <Configurations>Debug;Release;Development;Analyze</Configurations>
     <Nullable>enable</Nullable>
     <AllowUnsafeBlocks>true</AllowUnsafeBlocks>
     <GenerateDocumentationFile>True</GenerateDocumentationFile>
@@ -30,6 +30,12 @@
     <DocumentationFile></DocumentationFile>
     <Optimize>true</Optimize>
   </PropertyGroup>
+  <PropertyGroup Condition=" '$(Configuration)|$(Platform)' == 'Analyze|AnyCPU' ">
+    <RunAnalyzersDuringBuild>True</RunAnalyzersDuringBuild>
+    <TreatWarningsAsErrors>false</TreatWarningsAsErrors>
+    <GenerateDocumentationFile>True</GenerateDocumentationFile>
+    <DocumentationFile></DocumentationFile>
+  </PropertyGroup>
 
   <ItemGroup>
     <Compile Include="..\MetaData.cs">
diff --git a/Engine/Source/Programs/Shared/EpicGames.UHT/EpicGames.UHT.csproj b/Engine/Source/Programs/Shared/EpicGames.UHT/EpicGames.UHT.csproj
--- a/Engine/Source/Programs/Shared/EpicGames.UHT/EpicGames.UHT.csproj
+++ b/Engine/Source/Programs/Shared/EpicGames.UHT/EpicGames.UHT.csproj
@@ -8,7 +8,7 @@
     <RootNamespace>EpicGames.UHT</RootNamespace>
     <GenerateAssemblyInfo>false</GenerateAssemblyInfo>
     <GenerateTargetFrameworkAttribute>false</GenerateTargetFrameworkAttribute>
-    <Configurations>Debug;Release;Development</Configurations>
+    <Configurations>Debug;Release;Development;Analyze</Configurations>
     <Nullable>enable</Nullable>
     <AllowUnsafeBlocks>true</AllowUnsafeBlocks>
     <GenerateDocumentationFile>True</GenerateDocumentationFile>
@@ -29,6 +29,12 @@
     <DocumentationFile></DocumentationFile>
     <Optimize>true</Optimize>
   </PropertyGroup>
+  <PropertyGroup Condition=" '$(Configuration)|$(Platform)' == 'Analyze|AnyCPU' ">
+    <RunAnalyzersDuringBuild>True</RunAnalyzersDuringBuild>
+    <TreatWarningsAsErrors>false</TreatWarningsAsErrors>
+    <GenerateDocumentationFile>True</GenerateDocumentationFile>
+    <DocumentationFile></DocumentationFile>
+  </PropertyGroup>
   
   <ItemGroup>
     <Compile Include="..\MetaData.cs">
diff --git a/Engine/Source/Programs/UnrealBuildTool.Tests/ArtifactsTest.cs b/Engine/Source/Programs/UnrealBuildTool.Tests/ArtifactsTest.cs
--- a/Engine/Source/Programs/UnrealBuildTool.Tests/ArtifactsTest.cs
+++ b/Engine/Source/Programs/UnrealBuildTool.Tests/ArtifactsTest.cs
@@ -55,7 +55,7 @@ namespace UnrealBuildToolTests
 		}
 
 		[TestMethod]
-		public async Task ArtifactBundleStorageTest1()
+		public async Task ArtifactBundleStorageTest1Async()
 		{
 			CancellationToken cancellationToken = default;
 
diff --git a/Engine/Source/Programs/UnrealBuildTool.Tests/PluginDescriptorTests.cs b/Engine/Source/Programs/UnrealBuildTool.Tests/PluginDescriptorTests.cs
--- a/Engine/Source/Programs/UnrealBuildTool.Tests/PluginDescriptorTests.cs
+++ b/Engine/Source/Programs/UnrealBuildTool.Tests/PluginDescriptorTests.cs
@@ -15,7 +15,7 @@ namespace UnrealBuildToolTests
 	[TestClass]
 	public class PluginDescriptorTests
 	{
-		private static DirectoryReference s_tempDirectory;
+		private static DirectoryReference? s_tempDirectory;
 
 		private static DirectoryReference CreateTempDir()
 		{
@@ -33,7 +33,7 @@ namespace UnrealBuildToolTests
 		[ClassCleanup]
 		public static void TearDown()
 		{
-			if (Directory.Exists(PluginDescriptorTests.s_tempDirectory.FullName))
+			if (s_tempDirectory != null && Directory.Exists(PluginDescriptorTests.s_tempDirectory.FullName))
 			{
 				Directory.Delete(PluginDescriptorTests.s_tempDirectory.FullName, true);
 			}
@@ -127,7 +127,7 @@ namespace UnrealBuildToolTests
 
 		private static string GetAbsolutePathToTempFile(string fileName)
 		{
-			if (String.IsNullOrEmpty(fileName))
+			if (String.IsNullOrEmpty(fileName) || s_tempDirectory == null)
 			{
 				return "";
 			}
@@ -139,7 +139,7 @@ namespace UnrealBuildToolTests
 			string extension = Path.GetExtension(fileName);
 			Debug.Assert(!String.IsNullOrEmpty(extension));
 			Debug.Assert(extension == ".uplugin");
-			string inputFile = Path.Join(s_tempDirectory.FullName, fileName);
+			string inputFile = Path.Join(s_tempDirectory?.FullName, fileName);
 			FileReference inputFileReference = new FileReference(inputFile);
 			File.WriteAllText(inputFile, fileContent);
 			return inputFileReference;
diff --git a/Engine/Source/Programs/UnrealBuildTool.Tests/PreprocessorStateTests.cs b/Engine/Source/Programs/UnrealBuildTool.Tests/PreprocessorStateTests.cs
--- a/Engine/Source/Programs/UnrealBuildTool.Tests/PreprocessorStateTests.cs
+++ b/Engine/Source/Programs/UnrealBuildTool.Tests/PreprocessorStateTests.cs
@@ -99,7 +99,7 @@ namespace UnrealBuildToolTests
 			{
 				result.AppendFormat("+b {0}\n", branch);
 			}
-			foreach (KeyValuePair<Identifier, PreprocessorMacro> macro in transform.RequiredMacros)
+			foreach (KeyValuePair<Identifier, PreprocessorMacro?> macro in transform.RequiredMacros)
 			{
 				if (macro.Value == null)
 				{
@@ -110,7 +110,7 @@ namespace UnrealBuildToolTests
 					result.AppendFormat("=m {0}\n", macro.Value.ToString().TrimEnd());
 				}
 			}
-			foreach (KeyValuePair<Identifier, PreprocessorMacro> macro in transform.NewMacros)
+			foreach (KeyValuePair<Identifier, PreprocessorMacro?> macro in transform.NewMacros)
 			{
 				if (macro.Value == null)
 				{
@@ -124,7 +124,7 @@ namespace UnrealBuildToolTests
 			return result.ToString();
 		}
 
-		static PreprocessorMacro ParseMacro(string name, List<string> parameters, string value)
+		static PreprocessorMacro ParseMacro(string name, List<string>? parameters, string value)
 		{
 			List<Token> tokens = new List<Token>();
 
diff --git a/Engine/Source/Programs/UnrealBuildTool.Tests/PreprocessorTests.cs b/Engine/Source/Programs/UnrealBuildTool.Tests/PreprocessorTests.cs
--- a/Engine/Source/Programs/UnrealBuildTool.Tests/PreprocessorTests.cs
+++ b/Engine/Source/Programs/UnrealBuildTool.Tests/PreprocessorTests.cs
@@ -1,8 +1,10 @@
 // Copyright Epic Games, Inc. All Rights Reserved.
 
 using System.Collections.Generic;
+using System.IO;
 using System.Text;
 using Microsoft.VisualStudio.TestTools.UnitTesting;
+using UnrealBuildBase;
 using UnrealBuildTool;
 
 namespace UnrealBuildToolTests
@@ -220,17 +222,18 @@ namespace UnrealBuildToolTests
 		/// </summary>
 		/// <param name="fragment">The code fragment to preprocess</param>
 		/// <param name="expectedResult">The expected sequence of tokens, as a string. Null to indicate that the input is invalid, and an exception is expected.</param>
-		static void RunTest(string fragment, string expectedResult)
+		static void RunTest(string fragment, string? expectedResult)
 		{
 			string[] lines = fragment.Split('\n');
 
-			SourceFile file = new SourceFile(null, TokenReader.GetNullTerminatedByteArray(fragment));
+			SourceFile file = new SourceFile(FileItem.GetItemByPath(Path.Combine(Path.GetTempPath(), "temp.cpp")), TokenReader.GetNullTerminatedByteArray(fragment));
 
 			Preprocessor instance = new Preprocessor();
 
-			string result;
+			string? result;
 			try
 			{
+				PreprocessorFileContext context = new(file, null);
 				List<Token> outputTokens = new List<Token>();
 				for (int markupIdx = 0; markupIdx < file.Markup.Length; markupIdx++)
 				{
@@ -255,12 +258,12 @@ namespace UnrealBuildToolTests
 								tokens.Add(reader.Current);
 							}
 
-							instance.ExpandMacros(tokens, outputTokens, false, null);
+							instance.ExpandMacros(tokens, outputTokens, false, context);
 						}
 					}
 					else
 					{
-						instance.ParseMarkup(markup.Type, markup.Tokens, null);
+						instance.ParseMarkup(markup.Type, markup.Tokens ?? new(), context);
 					}
 				}
 				result = Token.Format(outputTokens);
diff --git a/Engine/Source/Programs/UnrealBuildTool.Tests/SourceFileTests.cs b/Engine/Source/Programs/UnrealBuildTool.Tests/SourceFileTests.cs
--- a/Engine/Source/Programs/UnrealBuildTool.Tests/SourceFileTests.cs
+++ b/Engine/Source/Programs/UnrealBuildTool.Tests/SourceFileTests.cs
@@ -25,7 +25,7 @@ namespace UnrealBuildToolTests
 			baseDirectories.Add(DirectoryReference.Combine(Unreal.EngineSourceDirectory, "Developer"));
 			baseDirectories.Add(DirectoryReference.Combine(Unreal.EngineSourceDirectory, "Editor"));
 
-			foreach (FileReference pluginFile in PluginsBase.EnumeratePlugins((FileReference)null))
+			foreach (FileReference pluginFile in PluginsBase.EnumeratePlugins((FileReference?)null))
 			{
 				DirectoryReference pluginSourceDir = DirectoryReference.Combine(pluginFile.Directory, "Source");
 				if (DirectoryReference.Exists(pluginSourceDir))
@@ -57,7 +57,7 @@ namespace UnrealBuildToolTests
 				}
 			}
 
-			List<SourceFile> readSourceFiles = new List<SourceFile>();
+			List<SourceFile>? readSourceFiles = new List<SourceFile>();
 			using (GlobalTracer.Instance.BuildSpan("Reading source file data").StartActive())
 			{
 				using (BinaryArchiveReader reader = new BinaryArchiveReader(tempDataFile))
diff --git a/Engine/Source/Programs/UnrealBuildTool.Tests/UnrealBuildTool.Tests.csproj b/Engine/Source/Programs/UnrealBuildTool.Tests/UnrealBuildTool.Tests.csproj
--- a/Engine/Source/Programs/UnrealBuildTool.Tests/UnrealBuildTool.Tests.csproj
+++ b/Engine/Source/Programs/UnrealBuildTool.Tests/UnrealBuildTool.Tests.csproj
@@ -2,10 +2,19 @@
 
   <PropertyGroup>
     <TargetFramework>net6.0</TargetFramework>
-    <Optimize>False</Optimize>
+    <TargetFramework Condition="'$(UE_DOTNET_VERSION)' != ''">$(UE_DOTNET_VERSION)</TargetFramework>
+    <Configurations>Debug;Release;Development;Analyze</Configurations>
+    <Nullable>enable</Nullable>
     <IsPackable>false</IsPackable>
     <OutputPath>..\..\..\Binaries\DotNET\UnrealBuildTool.Tests</OutputPath>
+    <NoWarn>1591</NoWarn>
   </PropertyGroup>
+  <PropertyGroup Condition=" '$(Configuration)|$(Platform)' == 'Analyze|AnyCPU' ">
+    <RunAnalyzersDuringBuild>True</RunAnalyzersDuringBuild>
+    <TreatWarningsAsErrors>false</TreatWarningsAsErrors>
+    <GenerateDocumentationFile>True</GenerateDocumentationFile>
+    <DocumentationFile></DocumentationFile>
+  </PropertyGroup>  
 
   <ItemGroup>
     <PackageReference Include="Microsoft.NET.Test.Sdk" Version="17.3.0" />
