Horde: Fix some static analysis warnings in tests.

#rnx

diff --git a/Engine/Source/Programs/UnrealBuildTool.Tests/LogEventMatcherTests.cs b/Engine/Source/Programs/UnrealBuildTool.Tests/LogEventMatcherTests.cs
--- a/Engine/Source/Programs/UnrealBuildTool.Tests/LogEventMatcherTests.cs
+++ b/Engine/Source/Programs/UnrealBuildTool.Tests/LogEventMatcherTests.cs
@@ -179,7 +179,7 @@ namespace UnrealBuildToolTests
 
 				LogValue noteProperty1 = logEvents[2].GetProperty<LogValue>("file");
 				Assert.AreEqual(LogValueType.SourceFile, noteProperty1.Type);
-				Assert.AreEqual(@"Engine\Plugins\Editor\EditorScriptingUtilities\Source\EditorScriptingUtilities\Public\EditorLevelLibrary.h", noteProperty1.Properties![LogEventPropertyName.File].ToString()!);
+				Assert.AreEqual(@"Engine\Plugins\Editor\EditorScriptingUtilities\Source\EditorScriptingUtilities\Public\EditorLevelLibrary.h", noteProperty1.Properties![LogEventPropertyName.File]?.ToString()!);
 			}
 		}
 
@@ -673,10 +673,10 @@ namespace UnrealBuildToolTests
 			CheckEventGroup(logEvents, 0, 2, LogLevel.Error, KnownLogEvents.Linker_UndefinedSymbol);
 
 			LogValue symbolProperty = (LogValue)logEvents[0].GetProperty("symbol");
-			Assert.AreEqual("Foo::Bar", symbolProperty.Properties![LogEventPropertyName.Identifier].ToString());
+			Assert.AreEqual("Foo::Bar", symbolProperty.Properties![LogEventPropertyName.Identifier]?.ToString());
 
 			LogValue symbolProperty2 = (LogValue)logEvents[1].GetProperty("symbol");
-			Assert.AreEqual("Foo::Bar2", symbolProperty2.Properties![LogEventPropertyName.Identifier].ToString());
+			Assert.AreEqual("Foo::Bar2", symbolProperty2.Properties![LogEventPropertyName.Identifier]?.ToString());
 		}
 
 		[TestMethod]
@@ -694,7 +694,7 @@ namespace UnrealBuildToolTests
 			CheckEventGroup(logEvents, 0, 4, LogLevel.Error, KnownLogEvents.Linker_UndefinedSymbol);
 
 			LogValue symbolProperty = (LogValue)logEvents[0].GetProperty("symbol");
-			Assert.AreEqual("USkeleton::GetBlendProfile", symbolProperty.Properties![LogEventPropertyName.Identifier].ToString());
+			Assert.AreEqual("USkeleton::GetBlendProfile", symbolProperty.Properties![LogEventPropertyName.Identifier]?.ToString());
 		}
 
 		[TestMethod]
@@ -711,7 +711,7 @@ namespace UnrealBuildToolTests
 			CheckEventGroup(logEvents, 0, 3, LogLevel.Error, KnownLogEvents.Linker_DuplicateSymbol);
 
 			LogValue symbolProperty = (LogValue)logEvents[0].GetProperty("symbol");
-			Assert.AreEqual("IMPLEMENT_MODULE_DataInterface", symbolProperty.Properties![LogEventPropertyName.Identifier].ToString());
+			Assert.AreEqual("IMPLEMENT_MODULE_DataInterface", symbolProperty.Properties![LogEventPropertyName.Identifier]?.ToString());
 		}
 
 		[TestMethod]
@@ -728,7 +728,7 @@ namespace UnrealBuildToolTests
 
 			LogValue symbolProperty = (LogValue)logEvents[1].GetProperty("symbol");
 			Assert.AreEqual(LogValueType.Symbol, symbolProperty.Type);
-			Assert.AreEqual("Foo::Bar", symbolProperty.Properties![LogEventPropertyName.Identifier].ToString());
+			Assert.AreEqual("Foo::Bar", symbolProperty.Properties![LogEventPropertyName.Identifier]?.ToString());
 		}
 
 		[TestMethod]
@@ -782,7 +782,7 @@ namespace UnrealBuildToolTests
 			CheckEventGroup(logEvents, 0, 1, LogLevel.Error, KnownLogEvents.Linker_UndefinedSymbol);
 
 			LogValue symbolProperty = (LogValue)logEvents[0].GetProperty("symbol");
-			Assert.AreEqual("Foo::Bar", symbolProperty.Properties![LogEventPropertyName.Identifier].ToString());
+			Assert.AreEqual("Foo::Bar", symbolProperty.Properties![LogEventPropertyName.Identifier]?.ToString());
 		}
 
 		[TestMethod]
@@ -815,7 +815,7 @@ namespace UnrealBuildToolTests
 			CheckEventGroup(logEvents, 0, 1, LogLevel.Error, KnownLogEvents.Linker_UndefinedSymbol);
 
 			LogValue symbolProperty = logEvents[0].GetProperty<LogValue>("symbol");
-			Assert.AreEqual("Foo::Bar", symbolProperty.Properties![LogEventPropertyName.Identifier].ToString());
+			Assert.AreEqual("Foo::Bar", symbolProperty.Properties![LogEventPropertyName.Identifier]?.ToString());
 		}
 
 		[TestMethod]
@@ -830,7 +830,7 @@ namespace UnrealBuildToolTests
 			CheckEventGroup(logEvents, 0, 1, LogLevel.Error, KnownLogEvents.Linker_UndefinedSymbol);
 
 			LogValue symbolProperty = logEvents[0].GetProperty<LogValue>("symbol");
-			Assert.AreEqual("Foo::Bar", symbolProperty.Properties![LogEventPropertyName.Identifier].ToString());
+			Assert.AreEqual("Foo::Bar", symbolProperty.Properties![LogEventPropertyName.Identifier]?.ToString());
 		}
 
 		[TestMethod]
@@ -845,7 +845,7 @@ namespace UnrealBuildToolTests
 			CheckEventGroup(logEvents, 0, 1, LogLevel.Error, KnownLogEvents.Linker_UndefinedSymbol);
 
 			LogValue symbolProperty = logEvents[0].GetProperty<LogValue>("symbol");
-			Assert.AreEqual("Foo::Bar", symbolProperty.Properties![LogEventPropertyName.Identifier].ToString());
+			Assert.AreEqual("Foo::Bar", symbolProperty.Properties![LogEventPropertyName.Identifier]?.ToString());
 		}
 
 		[TestMethod]
