Fix static analysis warning in test.

#rnx

diff --git a/Engine/Source/Programs/UnrealBuildTool.Tests/ConfigFileTests.cs b/Engine/Source/Programs/UnrealBuildTool.Tests/ConfigFileTests.cs
--- a/Engine/Source/Programs/UnrealBuildTool.Tests/ConfigFileTests.cs
+++ b/Engine/Source/Programs/UnrealBuildTool.Tests/ConfigFileTests.cs
@@ -74,6 +74,7 @@ namespace UnrealBuildToolTests
 				FileReference.MakeWriteable(ConfigFileLocation);
 				ConfigFile.Write(ConfigFileLocation);
 
+				Assert.IsNotNull(TemporaryConfigFile.TemporaryFile);
 				string FileContents = File.ReadAllText(TemporaryConfigFile.TemporaryFile);
 				Assert.AreEqual(TestContents.Trim(), FileContents.Trim(), "The file contents after writing do not match the expected contents.");
 			}
