Fix requring UnrealBuildTool.Tests be output to the binaries directory

#rnx

diff --git a/Engine/Source/Programs/UnrealBuildTool.Tests/SourceFileTests.cs b/Engine/Source/Programs/UnrealBuildTool.Tests/SourceFileTests.cs
--- a/Engine/Source/Programs/UnrealBuildTool.Tests/SourceFileTests.cs
+++ b/Engine/Source/Programs/UnrealBuildTool.Tests/SourceFileTests.cs
@@ -20,6 +20,8 @@ namespace UnrealBuildToolTests
 		[TestMethod]
 		public void Run()
 		{
+			UnrealHelper.InitializePath();
+
 			List<DirectoryReference> baseDirectories = new List<DirectoryReference>();
 			baseDirectories.Add(DirectoryReference.Combine(Unreal.EngineSourceDirectory, "Runtime"));
 			baseDirectories.Add(DirectoryReference.Combine(Unreal.EngineSourceDirectory, "Developer"));
diff --git a/Engine/Source/Programs/UnrealBuildTool.Tests/UnrealHeaderToolTests.cs b/Engine/Source/Programs/UnrealBuildTool.Tests/UnrealHeaderToolTests.cs
--- a/Engine/Source/Programs/UnrealBuildTool.Tests/UnrealHeaderToolTests.cs
+++ b/Engine/Source/Programs/UnrealBuildTool.Tests/UnrealHeaderToolTests.cs
@@ -1,13 +1,9 @@
 // Copyright Epic Games, Inc. All Rights Reserved.
 
-#if DISABLE
-using System.Reflection;
 using EpicGames.Core;
 using EpicGames.UHT.Utils;
 using Microsoft.Extensions.Logging;
-using UnrealBuildBase;
 using UnrealBuildTool.Modes;
-#endif
 using Microsoft.VisualStudio.TestTools.UnitTesting;
 
 namespace UnrealBuildToolTests
@@ -16,11 +12,10 @@ namespace UnrealBuildToolTests
 	public class UnrealHeaderToolTests
 	{
 		[TestMethod]
+		[Ignore]
 		public void Run()
 		{
-#if DISABLE
-			FileReference AssemblyFile = new FileReference(Assembly.GetExecutingAssembly().Location);
-			Unreal.LocationOverride.RootDirectory = DirectoryReference.Combine(AssemblyFile.Directory, "../../../../..");
+			UnrealHelper.InitializePath();
 
 			string[] Arguments = System.Array.Empty<string>();
 			CommandLineArguments CommandLineArguments = new CommandLineArguments(Arguments);
@@ -35,7 +30,6 @@ namespace UnrealBuildToolTests
 			// Run the tests
 			using ILoggerFactory factory = LoggerFactory.Create(x => x.AddEpicDefault());
 			Assert.IsTrue(UhtTestHarness.RunTests(Tables, Config, Options, factory.CreateLogger<UhtTestHarness>()));
-#endif
 		}
 	}
 }
diff --git a/Engine/Source/Programs/UnrealBuildTool.Tests/UnrealHelper.cs b/Engine/Source/Programs/UnrealBuildTool.Tests/UnrealHelper.cs
new file mode 100644
--- /dev/null
+++ b/Engine/Source/Programs/UnrealBuildTool.Tests/UnrealHelper.cs
@@ -0,0 +1,61 @@
+// Copyright Epic Games, Inc. All Rights Reserved.
+
+using System;
+using System.IO;
+using System.Reflection;
+using EpicGames.Core;
+using UnrealBuildBase;
+
+namespace UnrealBuildToolTests
+{
+	static class UnrealHelper
+	{
+		public static void InitializePath()
+		{
+			if (Unreal.LocationOverride.RootDirectory != null)
+			{
+				return;
+			}
+
+			// Use the EntryAssembly (the application path), rather than the ExecutingAssembly (the library path)
+			string AssemblyLocation = Assembly.GetEntryAssembly()!.GetOriginalLocation();
+
+			DirectoryReference? FoundRootDirectory = DirectoryReference.FindCorrectCase(DirectoryReference.FromString(AssemblyLocation)!);
+
+			// Search up through the directory tree for the deepest instance of the sub-path "Engine/Source/Programs"
+			while (FoundRootDirectory != null)
+			{
+				if (String.Equals("Programs", FoundRootDirectory.GetDirectoryName(), StringComparison.OrdinalIgnoreCase))
+				{
+					FoundRootDirectory = FoundRootDirectory.ParentDirectory;
+					if (FoundRootDirectory != null && String.Equals("Source", FoundRootDirectory.GetDirectoryName(), StringComparison.OrdinalIgnoreCase))
+					{
+						FoundRootDirectory = FoundRootDirectory.ParentDirectory;
+						if (FoundRootDirectory != null && String.Equals("Engine", FoundRootDirectory.GetDirectoryName(), StringComparison.OrdinalIgnoreCase))
+						{
+							FoundRootDirectory = FoundRootDirectory.ParentDirectory;
+							break;
+						}
+						continue;
+					}
+					continue;
+				}
+				FoundRootDirectory = FoundRootDirectory.ParentDirectory;
+			}
+
+			if (FoundRootDirectory == null)
+			{
+				throw new Exception($"This code requires that applications using it are launched from a path containing \"Engine/Source/Programs\". This application was launched from {Path.GetDirectoryName(AssemblyLocation)}");
+			}
+
+			// Confirm that we've found a valid root directory, by testing for the existence of a well-known file
+			FileReference ExpectedExistingFile = FileReference.Combine(FoundRootDirectory, "Engine", "Build", "Build.version");
+			if (!FileReference.Exists(ExpectedExistingFile))
+			{
+				throw new Exception($"Expected file \"Engine/Build/Build.version\" was not found at {ExpectedExistingFile.FullName}");
+			}
+
+			Unreal.LocationOverride.RootDirectory = FoundRootDirectory;
+		}
+	}
+}
