Update MSTest packages

#rnx

diff --git a/Engine/Source/Programs/AutomationTool.Tests/AutomationTool.Tests.csproj b/Engine/Source/Programs/AutomationTool.Tests/AutomationTool.Tests.csproj
--- a/Engine/Source/Programs/AutomationTool.Tests/AutomationTool.Tests.csproj
+++ b/Engine/Source/Programs/AutomationTool.Tests/AutomationTool.Tests.csproj
@@ -10,9 +10,9 @@
   </PropertyGroup>
 
   <ItemGroup>
-    <PackageReference Include="Microsoft.NET.Test.Sdk" Version="17.5.0" />
-    <PackageReference Include="MSTest.TestAdapter" Version="2.2.10" />
-    <PackageReference Include="MSTest.TestFramework" Version="2.2.10" />
+    <PackageReference Include="Microsoft.NET.Test.Sdk" Version="17.11.1" />
+    <PackageReference Include="MSTest.TestAdapter" Version="3.6.0" />
+    <PackageReference Include="MSTest.TestFramework" Version="3.6.0" />
     <PackageReference Include="coverlet.collector" Version="6.0.2">
       <PrivateAssets>all</PrivateAssets>
       <IncludeAssets>runtime; build; native; contentfiles; analyzers; buildtransitive</IncludeAssets>
diff --git a/Engine/Source/Programs/Horde/Drivers/JobDriver.Tests/JobDriver.Tests.csproj b/Engine/Source/Programs/Horde/Drivers/JobDriver.Tests/JobDriver.Tests.csproj
--- a/Engine/Source/Programs/Horde/Drivers/JobDriver.Tests/JobDriver.Tests.csproj
+++ b/Engine/Source/Programs/Horde/Drivers/JobDriver.Tests/JobDriver.Tests.csproj
@@ -20,10 +20,10 @@
 
   <ItemGroup>
     <PackageReference Include="Microsoft.Extensions.Logging.Abstractions" Version="8.0.0" />
-    <PackageReference Include="Microsoft.NET.Test.Sdk" Version="17.8.0" />
+    <PackageReference Include="Microsoft.NET.Test.Sdk" Version="17.11.1" />
     <PackageReference Include="Moq" Version="4.20.70" />
-    <PackageReference Include="MSTest.TestAdapter" Version="3.1.1" />
-    <PackageReference Include="MSTest.TestFramework" Version="3.1.1" />
+    <PackageReference Include="MSTest.TestAdapter" Version="3.6.0" />
+    <PackageReference Include="MSTest.TestFramework" Version="3.6.0" />
     <PackageReference Include="coverlet.collector" Version="6.0.2">
       <PrivateAssets>all</PrivateAssets>
       <IncludeAssets>runtime; build; native; contentfiles; analyzers; buildtransitive</IncludeAssets>
diff --git a/Engine/Source/Programs/Horde/HordeAgent.Tests/HordeAgent.Tests.csproj b/Engine/Source/Programs/Horde/HordeAgent.Tests/HordeAgent.Tests.csproj
--- a/Engine/Source/Programs/Horde/HordeAgent.Tests/HordeAgent.Tests.csproj
+++ b/Engine/Source/Programs/Horde/HordeAgent.Tests/HordeAgent.Tests.csproj
@@ -20,10 +20,10 @@
 
   <ItemGroup>
     <PackageReference Include="Microsoft.Extensions.Logging.Abstractions" Version="8.0.0" />
-    <PackageReference Include="Microsoft.NET.Test.Sdk" Version="17.8.0" />
+    <PackageReference Include="Microsoft.NET.Test.Sdk" Version="17.11.1" />
     <PackageReference Include="Moq" Version="4.20.70" />
-    <PackageReference Include="MSTest.TestAdapter" Version="3.1.1" />
-    <PackageReference Include="MSTest.TestFramework" Version="3.1.1" />
+    <PackageReference Include="MSTest.TestAdapter" Version="3.6.0" />
+    <PackageReference Include="MSTest.TestFramework" Version="3.6.0" />
     <PackageReference Include="coverlet.collector" Version="6.0.2">
       <PrivateAssets>all</PrivateAssets>
       <IncludeAssets>runtime; build; native; contentfiles; analyzers; buildtransitive</IncludeAssets>
diff --git a/Engine/Source/Programs/Horde/HordeServer.Tests.Shared/HordeServer.Tests.Shared.csproj b/Engine/Source/Programs/Horde/HordeServer.Tests.Shared/HordeServer.Tests.Shared.csproj
--- a/Engine/Source/Programs/Horde/HordeServer.Tests.Shared/HordeServer.Tests.Shared.csproj
+++ b/Engine/Source/Programs/Horde/HordeServer.Tests.Shared/HordeServer.Tests.Shared.csproj
@@ -23,10 +23,10 @@
 
   <ItemGroup>
     <PackageReference Include="Microsoft.AspNetCore.Mvc.Testing" Version="6.0.25" />
-    <PackageReference Include="Microsoft.NET.Test.Sdk" Version="17.8.0" />
+    <PackageReference Include="Microsoft.NET.Test.Sdk" Version="17.11.1" />
     <PackageReference Include="Moq" Version="4.20.70" />
-    <PackageReference Include="MSTest.TestAdapter" Version="3.1.1" />
-    <PackageReference Include="MSTest.TestFramework" Version="3.1.1" />
+    <PackageReference Include="MSTest.TestAdapter" Version="3.6.0" />
+    <PackageReference Include="MSTest.TestFramework" Version="3.6.0" />
     <PackageReference Include="Microsoft.Extensions.Caching.Memory" Version="8.0.0" />
     <PackageReference Include="coverlet.collector" Version="6.0.2">
       <PrivateAssets>all</PrivateAssets>
diff --git a/Engine/Source/Programs/Horde/HordeServer.Tests/HordeServer.Tests.csproj b/Engine/Source/Programs/Horde/HordeServer.Tests/HordeServer.Tests.csproj
--- a/Engine/Source/Programs/Horde/HordeServer.Tests/HordeServer.Tests.csproj
+++ b/Engine/Source/Programs/Horde/HordeServer.Tests/HordeServer.Tests.csproj
@@ -19,10 +19,10 @@
 
   <ItemGroup>
     <PackageReference Include="Microsoft.AspNetCore.Mvc.Testing" Version="6.0.25" />
-    <PackageReference Include="Microsoft.NET.Test.Sdk" Version="17.8.0" />
+    <PackageReference Include="Microsoft.NET.Test.Sdk" Version="17.11.1" />
     <PackageReference Include="Moq" Version="4.20.70" />
-    <PackageReference Include="MSTest.TestAdapter" Version="3.1.1" />
-    <PackageReference Include="MSTest.TestFramework" Version="3.1.1" />
+    <PackageReference Include="MSTest.TestAdapter" Version="3.6.0" />
+    <PackageReference Include="MSTest.TestFramework" Version="3.6.0" />
     <PackageReference Include="Microsoft.Extensions.Caching.Memory" Version="8.0.0" />
     <PackageReference Include="coverlet.collector" Version="6.0.2">
       <PrivateAssets>all</PrivateAssets>
diff --git a/Engine/Source/Programs/Horde/Plugins/Analytics/HordeServer.Analytics.Tests/HordeServer.Analytics.Tests.csproj b/Engine/Source/Programs/Horde/Plugins/Analytics/HordeServer.Analytics.Tests/HordeServer.Analytics.Tests.csproj
--- a/Engine/Source/Programs/Horde/Plugins/Analytics/HordeServer.Analytics.Tests/HordeServer.Analytics.Tests.csproj
+++ b/Engine/Source/Programs/Horde/Plugins/Analytics/HordeServer.Analytics.Tests/HordeServer.Analytics.Tests.csproj
@@ -23,9 +23,9 @@
       <PrivateAssets>all</PrivateAssets>
       <IncludeAssets>runtime; build; native; contentfiles; analyzers; buildtransitive</IncludeAssets>
     </PackageReference>
-    <PackageReference Include="Microsoft.NET.Test.Sdk" Version="17.8.0" />
-    <PackageReference Include="MSTest.TestAdapter" Version="3.1.1" />
-    <PackageReference Include="MSTest.TestFramework" Version="3.1.1" />
+    <PackageReference Include="Microsoft.NET.Test.Sdk" Version="17.11.1" />
+    <PackageReference Include="MSTest.TestAdapter" Version="3.6.0" />
+    <PackageReference Include="MSTest.TestFramework" Version="3.6.0" />
   </ItemGroup>
 
   <ItemGroup>
diff --git a/Engine/Source/Programs/Horde/Plugins/Build/HordeServer.Build.Tests/HordeServer.Build.Tests.csproj b/Engine/Source/Programs/Horde/Plugins/Build/HordeServer.Build.Tests/HordeServer.Build.Tests.csproj
--- a/Engine/Source/Programs/Horde/Plugins/Build/HordeServer.Build.Tests/HordeServer.Build.Tests.csproj
+++ b/Engine/Source/Programs/Horde/Plugins/Build/HordeServer.Build.Tests/HordeServer.Build.Tests.csproj
@@ -24,9 +24,9 @@
       <PrivateAssets>all</PrivateAssets>
       <IncludeAssets>runtime; build; native; contentfiles; analyzers; buildtransitive</IncludeAssets>
     </PackageReference>
-    <PackageReference Include="Microsoft.NET.Test.Sdk" Version="17.8.0" />
-    <PackageReference Include="MSTest.TestAdapter" Version="3.1.1" />
-    <PackageReference Include="MSTest.TestFramework" Version="3.1.1" />
+    <PackageReference Include="Microsoft.NET.Test.Sdk" Version="17.11.1" />
+    <PackageReference Include="MSTest.TestAdapter" Version="3.6.0" />
+    <PackageReference Include="MSTest.TestFramework" Version="3.6.0" />
     <PackageReference Include="Microsoft.AspNetCore.Mvc.Testing" Version="6.0.25" />
     <PackageReference Include="Moq" Version="4.20.70" />
     <PackageReference Include="Microsoft.Extensions.Caching.Memory" Version="8.0.0" />
diff --git a/Engine/Source/Programs/Horde/Plugins/Compute/HordeServer.Compute.Tests/HordeServer.Compute.Tests.csproj b/Engine/Source/Programs/Horde/Plugins/Compute/HordeServer.Compute.Tests/HordeServer.Compute.Tests.csproj
--- a/Engine/Source/Programs/Horde/Plugins/Compute/HordeServer.Compute.Tests/HordeServer.Compute.Tests.csproj
+++ b/Engine/Source/Programs/Horde/Plugins/Compute/HordeServer.Compute.Tests/HordeServer.Compute.Tests.csproj
@@ -24,9 +24,9 @@
       <PrivateAssets>all</PrivateAssets>
       <IncludeAssets>runtime; build; native; contentfiles; analyzers; buildtransitive</IncludeAssets>
     </PackageReference>
-    <PackageReference Include="Microsoft.NET.Test.Sdk" Version="17.8.0" />
-    <PackageReference Include="MSTest.TestAdapter" Version="3.1.1" />
-    <PackageReference Include="MSTest.TestFramework" Version="3.1.1" />
+    <PackageReference Include="Microsoft.NET.Test.Sdk" Version="17.11.1" />
+    <PackageReference Include="MSTest.TestAdapter" Version="3.6.0" />
+    <PackageReference Include="MSTest.TestFramework" Version="3.6.0" />
   </ItemGroup>
 
   <ItemGroup>
diff --git a/Engine/Source/Programs/Horde/Plugins/Ddc/HordeServer.Ddc.Tests/HordeServer.Ddc.Tests.csproj b/Engine/Source/Programs/Horde/Plugins/Ddc/HordeServer.Ddc.Tests/HordeServer.Ddc.Tests.csproj
--- a/Engine/Source/Programs/Horde/Plugins/Ddc/HordeServer.Ddc.Tests/HordeServer.Ddc.Tests.csproj
+++ b/Engine/Source/Programs/Horde/Plugins/Ddc/HordeServer.Ddc.Tests/HordeServer.Ddc.Tests.csproj
@@ -24,9 +24,9 @@
       <PrivateAssets>all</PrivateAssets>
       <IncludeAssets>runtime; build; native; contentfiles; analyzers; buildtransitive</IncludeAssets>
     </PackageReference>
-    <PackageReference Include="Microsoft.NET.Test.Sdk" Version="17.8.0" />
-    <PackageReference Include="MSTest.TestAdapter" Version="3.1.1" />
-    <PackageReference Include="MSTest.TestFramework" Version="3.1.1" />
+    <PackageReference Include="Microsoft.NET.Test.Sdk" Version="17.11.1" />
+    <PackageReference Include="MSTest.TestAdapter" Version="3.6.0" />
+    <PackageReference Include="MSTest.TestFramework" Version="3.6.0" />
   </ItemGroup>
 
   <ItemGroup>
diff --git a/Engine/Source/Programs/Horde/Plugins/Storage/HordeServer.Storage.Tests/HordeServer.Storage.Tests.csproj b/Engine/Source/Programs/Horde/Plugins/Storage/HordeServer.Storage.Tests/HordeServer.Storage.Tests.csproj
--- a/Engine/Source/Programs/Horde/Plugins/Storage/HordeServer.Storage.Tests/HordeServer.Storage.Tests.csproj
+++ b/Engine/Source/Programs/Horde/Plugins/Storage/HordeServer.Storage.Tests/HordeServer.Storage.Tests.csproj
@@ -24,9 +24,9 @@
       <PrivateAssets>all</PrivateAssets>
       <IncludeAssets>runtime; build; native; contentfiles; analyzers; buildtransitive</IncludeAssets>
     </PackageReference>
-    <PackageReference Include="Microsoft.NET.Test.Sdk" Version="17.8.0" />
-    <PackageReference Include="MSTest.TestAdapter" Version="3.1.1" />
-    <PackageReference Include="MSTest.TestFramework" Version="3.1.1" />
+    <PackageReference Include="Microsoft.NET.Test.Sdk" Version="17.11.1" />
+    <PackageReference Include="MSTest.TestAdapter" Version="3.6.0" />
+    <PackageReference Include="MSTest.TestFramework" Version="3.6.0" />
   </ItemGroup>
 
   <ItemGroup>
diff --git a/Engine/Source/Programs/Horde/Plugins/Tools/HordeServer.Tools.Tests/HordeServer.Tools.Tests.csproj b/Engine/Source/Programs/Horde/Plugins/Tools/HordeServer.Tools.Tests/HordeServer.Tools.Tests.csproj
--- a/Engine/Source/Programs/Horde/Plugins/Tools/HordeServer.Tools.Tests/HordeServer.Tools.Tests.csproj
+++ b/Engine/Source/Programs/Horde/Plugins/Tools/HordeServer.Tools.Tests/HordeServer.Tools.Tests.csproj
@@ -24,9 +24,9 @@
       <PrivateAssets>all</PrivateAssets>
       <IncludeAssets>runtime; build; native; contentfiles; analyzers; buildtransitive</IncludeAssets>
     </PackageReference>
-    <PackageReference Include="Microsoft.NET.Test.Sdk" Version="17.8.0" />
-    <PackageReference Include="MSTest.TestAdapter" Version="3.1.1" />
-    <PackageReference Include="MSTest.TestFramework" Version="3.1.1" />
+    <PackageReference Include="Microsoft.NET.Test.Sdk" Version="17.11.1" />
+    <PackageReference Include="MSTest.TestAdapter" Version="3.6.0" />
+    <PackageReference Include="MSTest.TestFramework" Version="3.6.0" />
   </ItemGroup>
 
   <ItemGroup>
diff --git a/Engine/Source/Programs/Shared/EpicGames.BuildGraph.Tests/EpicGames.BuildGraph.Tests.csproj b/Engine/Source/Programs/Shared/EpicGames.BuildGraph.Tests/EpicGames.BuildGraph.Tests.csproj
--- a/Engine/Source/Programs/Shared/EpicGames.BuildGraph.Tests/EpicGames.BuildGraph.Tests.csproj
+++ b/Engine/Source/Programs/Shared/EpicGames.BuildGraph.Tests/EpicGames.BuildGraph.Tests.csproj
@@ -45,9 +45,9 @@
   </ItemGroup>
 
   <ItemGroup>
-    <PackageReference Include="Microsoft.NET.Test.Sdk" Version="17.11.0" />
-    <PackageReference Include="MSTest.TestAdapter" Version="3.5.2" />
-    <PackageReference Include="MSTest.TestFramework" Version="3.5.2" />
+    <PackageReference Include="Microsoft.NET.Test.Sdk" Version="17.11.1" />
+    <PackageReference Include="MSTest.TestAdapter" Version="3.6.0" />
+    <PackageReference Include="MSTest.TestFramework" Version="3.6.0" />
     <PackageReference Include="coverlet.collector" Version="6.0.2">
       <PrivateAssets>all</PrivateAssets>
       <IncludeAssets>runtime; build; native; contentfiles; analyzers; buildtransitive</IncludeAssets>
diff --git a/Engine/Source/Programs/Shared/EpicGames.Core.Tests/EpicGames.Core.Tests.csproj b/Engine/Source/Programs/Shared/EpicGames.Core.Tests/EpicGames.Core.Tests.csproj
--- a/Engine/Source/Programs/Shared/EpicGames.Core.Tests/EpicGames.Core.Tests.csproj
+++ b/Engine/Source/Programs/Shared/EpicGames.Core.Tests/EpicGames.Core.Tests.csproj
@@ -48,9 +48,9 @@
   </ItemGroup>
 
   <ItemGroup>
-    <PackageReference Include="Microsoft.NET.Test.Sdk" Version="17.11.0" />
-    <PackageReference Include="MSTest.TestAdapter" Version="3.5.2" />
-    <PackageReference Include="MSTest.TestFramework" Version="3.5.2" />
+    <PackageReference Include="Microsoft.NET.Test.Sdk" Version="17.11.1" />
+    <PackageReference Include="MSTest.TestAdapter" Version="3.6.0" />
+    <PackageReference Include="MSTest.TestFramework" Version="3.6.0" />
     <PackageReference Include="coverlet.collector" Version="6.0.2">
       <PrivateAssets>all</PrivateAssets>
       <IncludeAssets>runtime; build; native; contentfiles; analyzers; buildtransitive</IncludeAssets>
diff --git a/Engine/Source/Programs/Shared/EpicGames.Horde.Tests/EpicGames.Horde.Tests.csproj b/Engine/Source/Programs/Shared/EpicGames.Horde.Tests/EpicGames.Horde.Tests.csproj
--- a/Engine/Source/Programs/Shared/EpicGames.Horde.Tests/EpicGames.Horde.Tests.csproj
+++ b/Engine/Source/Programs/Shared/EpicGames.Horde.Tests/EpicGames.Horde.Tests.csproj
@@ -30,10 +30,10 @@
   </PropertyGroup>
 
   <ItemGroup>
-    <PackageReference Include="Microsoft.NET.Test.Sdk" Version="17.11.0" />
+    <PackageReference Include="Microsoft.NET.Test.Sdk" Version="17.11.1" />
     <PackageReference Include="Microsoft.Extensions.Caching.Memory" Version="6.0.1" />
-    <PackageReference Include="MSTest.TestAdapter" Version="3.5.2" />
-    <PackageReference Include="MSTest.TestFramework" Version="3.5.2" />
+    <PackageReference Include="MSTest.TestAdapter" Version="3.6.0" />
+    <PackageReference Include="MSTest.TestFramework" Version="3.6.0" />
     <PackageReference Include="coverlet.collector" Version="6.0.2">
       <PrivateAssets>all</PrivateAssets>
       <IncludeAssets>runtime; build; native; contentfiles; analyzers; buildtransitive</IncludeAssets>
diff --git a/Engine/Source/Programs/Shared/EpicGames.IoHash.Tests/EpicGames.IoHash.Tests.csproj b/Engine/Source/Programs/Shared/EpicGames.IoHash.Tests/EpicGames.IoHash.Tests.csproj
--- a/Engine/Source/Programs/Shared/EpicGames.IoHash.Tests/EpicGames.IoHash.Tests.csproj
+++ b/Engine/Source/Programs/Shared/EpicGames.IoHash.Tests/EpicGames.IoHash.Tests.csproj
@@ -17,9 +17,9 @@
   </PropertyGroup>
 
   <ItemGroup>
-    <PackageReference Include="Microsoft.NET.Test.Sdk" Version="17.11.0" />
-    <PackageReference Include="MSTest.TestAdapter" Version="3.5.2" />
-    <PackageReference Include="MSTest.TestFramework" Version="3.5.2" />
+    <PackageReference Include="Microsoft.NET.Test.Sdk" Version="17.11.1" />
+    <PackageReference Include="MSTest.TestAdapter" Version="3.6.0" />
+    <PackageReference Include="MSTest.TestFramework" Version="3.6.0" />
     <PackageReference Include="coverlet.collector" Version="6.0.2">
       <PrivateAssets>all</PrivateAssets>
       <IncludeAssets>runtime; build; native; contentfiles; analyzers; buildtransitive</IncludeAssets>
diff --git a/Engine/Source/Programs/Shared/EpicGames.Perforce.Fixture/EpicGames.Perforce.Fixture.csproj b/Engine/Source/Programs/Shared/EpicGames.Perforce.Fixture/EpicGames.Perforce.Fixture.csproj
--- a/Engine/Source/Programs/Shared/EpicGames.Perforce.Fixture/EpicGames.Perforce.Fixture.csproj
+++ b/Engine/Source/Programs/Shared/EpicGames.Perforce.Fixture/EpicGames.Perforce.Fixture.csproj
@@ -33,9 +33,9 @@
 
   <ItemGroup>
     <PackageReference Include="Microsoft.Extensions.Logging.Console" Version="8.0.0" />
-    <PackageReference Include="Microsoft.NET.Test.Sdk" Version="17.11.0" />
-    <PackageReference Include="MSTest.TestAdapter" Version="3.5.2" />
-    <PackageReference Include="MSTest.TestFramework" Version="3.5.2" />
+    <PackageReference Include="Microsoft.NET.Test.Sdk" Version="17.11.1" />
+    <PackageReference Include="MSTest.TestAdapter" Version="3.6.0" />
+    <PackageReference Include="MSTest.TestFramework" Version="3.6.0" />
     <PackageReference Include="coverlet.collector" Version="6.0.2">
       <PrivateAssets>all</PrivateAssets>
       <IncludeAssets>runtime; build; native; contentfiles; analyzers; buildtransitive</IncludeAssets>
diff --git a/Engine/Source/Programs/Shared/EpicGames.Perforce.Managed.Tests/EpicGames.Perforce.Managed.Tests.csproj b/Engine/Source/Programs/Shared/EpicGames.Perforce.Managed.Tests/EpicGames.Perforce.Managed.Tests.csproj
--- a/Engine/Source/Programs/Shared/EpicGames.Perforce.Managed.Tests/EpicGames.Perforce.Managed.Tests.csproj
+++ b/Engine/Source/Programs/Shared/EpicGames.Perforce.Managed.Tests/EpicGames.Perforce.Managed.Tests.csproj
@@ -31,9 +31,9 @@
   </ItemGroup>
 
   <ItemGroup>
-    <PackageReference Include="Microsoft.NET.Test.Sdk" Version="17.11.0" />
-    <PackageReference Include="MSTest.TestAdapter" Version="3.5.2" />
-    <PackageReference Include="MSTest.TestFramework" Version="3.5.2" />
+    <PackageReference Include="Microsoft.NET.Test.Sdk" Version="17.11.1" />
+    <PackageReference Include="MSTest.TestAdapter" Version="3.6.0" />
+    <PackageReference Include="MSTest.TestFramework" Version="3.6.0" />
     <PackageReference Include="coverlet.collector" Version="6.0.2">
       <PrivateAssets>all</PrivateAssets>
       <IncludeAssets>runtime; build; native; contentfiles; analyzers; buildtransitive</IncludeAssets>
diff --git a/Engine/Source/Programs/Shared/EpicGames.Perforce.Tests/EpicGames.Perforce.Tests.csproj b/Engine/Source/Programs/Shared/EpicGames.Perforce.Tests/EpicGames.Perforce.Tests.csproj
--- a/Engine/Source/Programs/Shared/EpicGames.Perforce.Tests/EpicGames.Perforce.Tests.csproj
+++ b/Engine/Source/Programs/Shared/EpicGames.Perforce.Tests/EpicGames.Perforce.Tests.csproj
@@ -30,9 +30,9 @@
   </ItemGroup>
 
   <ItemGroup>
-    <PackageReference Include="Microsoft.NET.Test.Sdk" Version="17.11.0" />
-    <PackageReference Include="MSTest.TestAdapter" Version="3.5.2" />
-    <PackageReference Include="MSTest.TestFramework" Version="3.5.2" />
+    <PackageReference Include="Microsoft.NET.Test.Sdk" Version="17.11.1" />
+    <PackageReference Include="MSTest.TestAdapter" Version="3.6.0" />
+    <PackageReference Include="MSTest.TestFramework" Version="3.6.0" />
     <PackageReference Include="coverlet.collector" Version="6.0.2">
       <PrivateAssets>all</PrivateAssets>
       <IncludeAssets>runtime; build; native; contentfiles; analyzers; buildtransitive</IncludeAssets>
diff --git a/Engine/Source/Programs/Shared/EpicGames.Redis.Tests/EpicGames.Redis.Tests.csproj b/Engine/Source/Programs/Shared/EpicGames.Redis.Tests/EpicGames.Redis.Tests.csproj
--- a/Engine/Source/Programs/Shared/EpicGames.Redis.Tests/EpicGames.Redis.Tests.csproj
+++ b/Engine/Source/Programs/Shared/EpicGames.Redis.Tests/EpicGames.Redis.Tests.csproj
@@ -31,9 +31,9 @@
   </ItemGroup>
 
   <ItemGroup>
-    <PackageReference Include="Microsoft.NET.Test.Sdk" Version="17.11.0" />
-    <PackageReference Include="MSTest.TestAdapter" Version="3.5.2" />
-    <PackageReference Include="MSTest.TestFramework" Version="3.5.2" />
+    <PackageReference Include="Microsoft.NET.Test.Sdk" Version="17.11.1" />
+    <PackageReference Include="MSTest.TestAdapter" Version="3.6.0" />
+    <PackageReference Include="MSTest.TestFramework" Version="3.6.0" />
     <PackageReference Include="coverlet.collector" Version="6.0.2">
       <PrivateAssets>all</PrivateAssets>
       <IncludeAssets>runtime; build; native; contentfiles; analyzers; buildtransitive</IncludeAssets>
diff --git a/Engine/Source/Programs/Shared/EpicGames.Serialization.Tests/EpicGames.Serialization.Tests.csproj b/Engine/Source/Programs/Shared/EpicGames.Serialization.Tests/EpicGames.Serialization.Tests.csproj
--- a/Engine/Source/Programs/Shared/EpicGames.Serialization.Tests/EpicGames.Serialization.Tests.csproj
+++ b/Engine/Source/Programs/Shared/EpicGames.Serialization.Tests/EpicGames.Serialization.Tests.csproj
@@ -33,9 +33,9 @@
   </ItemGroup>
 
   <ItemGroup>
-    <PackageReference Include="Microsoft.NET.Test.Sdk" Version="17.11.0" />
-    <PackageReference Include="MSTest.TestAdapter" Version="3.5.2" />
-    <PackageReference Include="MSTest.TestFramework" Version="3.5.2" />
+    <PackageReference Include="Microsoft.NET.Test.Sdk" Version="17.11.1" />
+    <PackageReference Include="MSTest.TestAdapter" Version="3.6.0" />
+    <PackageReference Include="MSTest.TestFramework" Version="3.6.0" />
     <PackageReference Include="coverlet.collector" Version="6.0.2">
       <PrivateAssets>all</PrivateAssets>
       <IncludeAssets>runtime; build; native; contentfiles; analyzers; buildtransitive</IncludeAssets>
diff --git a/Engine/Source/Programs/Shared/EpicGames.Tracing.Tests/EpicGames.Tracing.Tests.csproj b/Engine/Source/Programs/Shared/EpicGames.Tracing.Tests/EpicGames.Tracing.Tests.csproj
--- a/Engine/Source/Programs/Shared/EpicGames.Tracing.Tests/EpicGames.Tracing.Tests.csproj
+++ b/Engine/Source/Programs/Shared/EpicGames.Tracing.Tests/EpicGames.Tracing.Tests.csproj
@@ -29,9 +29,9 @@
   </ItemGroup>
 
   <ItemGroup>
-    <PackageReference Include="Microsoft.NET.Test.Sdk" Version="17.11.0" />
-    <PackageReference Include="MSTest.TestAdapter" Version="3.5.2" />
-    <PackageReference Include="MSTest.TestFramework" Version="3.5.2" />
+    <PackageReference Include="Microsoft.NET.Test.Sdk" Version="17.11.1" />
+    <PackageReference Include="MSTest.TestAdapter" Version="3.6.0" />
+    <PackageReference Include="MSTest.TestFramework" Version="3.6.0" />
     <PackageReference Include="coverlet.collector" Version="6.0.2">
       <PrivateAssets>all</PrivateAssets>
       <IncludeAssets>runtime; build; native; contentfiles; analyzers; buildtransitive</IncludeAssets>
diff --git a/Engine/Source/Programs/UnrealBuildTool.Tests/ConfigFileTests.cs b/Engine/Source/Programs/UnrealBuildTool.Tests/ConfigFileTests.cs
--- a/Engine/Source/Programs/UnrealBuildTool.Tests/ConfigFileTests.cs
+++ b/Engine/Source/Programs/UnrealBuildTool.Tests/ConfigFileTests.cs
@@ -52,6 +52,8 @@ namespace UnrealBuildToolTests
 		[TestMethod]
 		public void ConfigFileTestsQuotations()
 		{
+			UnrealHelper.InitializePath();
+
 			const string TestSection = "MySection";
 			const string TestValueKey = "Name";
 			const string TestValue = "\"https://www.unrealengine.com/\"";
diff --git a/Engine/Source/Programs/UnrealBuildTool.Tests/UnrealBuildTool.Tests.csproj b/Engine/Source/Programs/UnrealBuildTool.Tests/UnrealBuildTool.Tests.csproj
--- a/Engine/Source/Programs/UnrealBuildTool.Tests/UnrealBuildTool.Tests.csproj
+++ b/Engine/Source/Programs/UnrealBuildTool.Tests/UnrealBuildTool.Tests.csproj
@@ -6,7 +6,6 @@
     <Configurations>Debug;Release;Development;Analyze</Configurations>
     <Nullable>enable</Nullable>
     <IsPackable>false</IsPackable>
-    <OutputPath>..\..\..\Binaries\DotNET\UnrealBuildTool.Tests</OutputPath>
     <NoWarn>1591</NoWarn>
   </PropertyGroup>
   <PropertyGroup Condition=" '$(Configuration)|$(Platform)' == 'Analyze|AnyCPU' ">
@@ -17,9 +16,9 @@
   </PropertyGroup>  
 
   <ItemGroup>
-    <PackageReference Include="Microsoft.NET.Test.Sdk" Version="17.11.0" />
-    <PackageReference Include="MSTest.TestAdapter" Version="3.5.2" />
-    <PackageReference Include="MSTest.TestFramework" Version="3.5.2" />
+    <PackageReference Include="Microsoft.NET.Test.Sdk" Version="17.11.1" />
+    <PackageReference Include="MSTest.TestAdapter" Version="3.6.0" />
+    <PackageReference Include="MSTest.TestFramework" Version="3.6.0" />
     <PackageReference Include="coverlet.collector" Version="6.0.2">
       <PrivateAssets>all</PrivateAssets>
       <IncludeAssets>runtime; build; native; contentfiles; analyzers; buildtransitive</IncludeAssets>
