Move UnrealBuildTool.Tests out of the Shared directory

#rnx

#Synced-CL 52229246

diff --git a/Engine/Source/Programs/Shared/UnrealBuildTool.Tests/ConfigFileTests.cs b/Engine/Source/Programs/Shared/UnrealBuildTool.Tests/ConfigFileTests.cs
deleted file mode 100644
diff --git a/Engine/Source/Programs/Shared/UnrealBuildTool.Tests/CppCompileWarningsTests.cs b/Engine/Source/Programs/Shared/UnrealBuildTool.Tests/CppCompileWarningsTests.cs
deleted file mode 100644
diff --git a/Engine/Source/Programs/Shared/UnrealBuildTool.Tests/JsonExtensionsTests.cs b/Engine/Source/Programs/Shared/UnrealBuildTool.Tests/JsonExtensionsTests.cs
deleted file mode 100644
diff --git a/Engine/Source/Programs/Shared/UnrealBuildTool.Tests/PluginDescriptorTests.cs b/Engine/Source/Programs/Shared/UnrealBuildTool.Tests/PluginDescriptorTests.cs
deleted file mode 100644
diff --git a/Engine/Source/Programs/Shared/UnrealBuildTool.Tests/PreprocessorExpressionTests.cs b/Engine/Source/Programs/Shared/UnrealBuildTool.Tests/PreprocessorExpressionTests.cs
deleted file mode 100644
diff --git a/Engine/Source/Programs/Shared/UnrealBuildTool.Tests/PreprocessorStateTests.cs b/Engine/Source/Programs/Shared/UnrealBuildTool.Tests/PreprocessorStateTests.cs
deleted file mode 100644
diff --git a/Engine/Source/Programs/Shared/UnrealBuildTool.Tests/PreprocessorTests.cs b/Engine/Source/Programs/Shared/UnrealBuildTool.Tests/PreprocessorTests.cs
deleted file mode 100644
diff --git a/Engine/Source/Programs/Shared/UnrealBuildTool.Tests/ProjectDescriptorTests.cs b/Engine/Source/Programs/Shared/UnrealBuildTool.Tests/ProjectDescriptorTests.cs
deleted file mode 100644
diff --git a/Engine/Source/Programs/Shared/UnrealBuildTool.Tests/SourceFileTests.cs b/Engine/Source/Programs/Shared/UnrealBuildTool.Tests/SourceFileTests.cs
deleted file mode 100644
diff --git a/Engine/Source/Programs/Shared/UnrealBuildTool.Tests/TargetRulesTests.cs b/Engine/Source/Programs/Shared/UnrealBuildTool.Tests/TargetRulesTests.cs
deleted file mode 100644
diff --git a/Engine/Source/Programs/Shared/UnrealBuildTool.Tests/TestUtilities/SafeTestDirectory.cs b/Engine/Source/Programs/Shared/UnrealBuildTool.Tests/TestUtilities/SafeTestDirectory.cs
deleted file mode 100644
diff --git a/Engine/Source/Programs/Shared/UnrealBuildTool.Tests/TestUtilities/SafeTestFile.cs b/Engine/Source/Programs/Shared/UnrealBuildTool.Tests/TestUtilities/SafeTestFile.cs
deleted file mode 100644
diff --git a/Engine/Source/Programs/Shared/UnrealBuildTool.Tests/TestUtilities/StaticTargetRulesInitializer.cs b/Engine/Source/Programs/Shared/UnrealBuildTool.Tests/TestUtilities/StaticTargetRulesInitializer.cs
deleted file mode 100644
diff --git a/Engine/Source/Programs/Shared/UnrealBuildTool.Tests/TestUtilities/TestBuildPlatform.cs b/Engine/Source/Programs/Shared/UnrealBuildTool.Tests/TestUtilities/TestBuildPlatform.cs
deleted file mode 100644
diff --git a/Engine/Source/Programs/Shared/UnrealBuildTool.Tests/TestUtilities/TestBuildPlatformFactory.cs b/Engine/Source/Programs/Shared/UnrealBuildTool.Tests/TestUtilities/TestBuildPlatformFactory.cs
deleted file mode 100644
diff --git a/Engine/Source/Programs/Shared/UnrealBuildTool.Tests/TestUtilities/TestBuildPlatformSDK.cs b/Engine/Source/Programs/Shared/UnrealBuildTool.Tests/TestUtilities/TestBuildPlatformSDK.cs
deleted file mode 100644
diff --git a/Engine/Source/Programs/Shared/UnrealBuildTool.Tests/TestUtilities/TestDirs.cs b/Engine/Source/Programs/Shared/UnrealBuildTool.Tests/TestUtilities/TestDirs.cs
deleted file mode 100644
diff --git a/Engine/Source/Programs/Shared/UnrealBuildTool.Tests/TestUtilities/TestLogger.cs b/Engine/Source/Programs/Shared/UnrealBuildTool.Tests/TestUtilities/TestLogger.cs
deleted file mode 100644
diff --git a/Engine/Source/Programs/Shared/UnrealBuildTool.Tests/TestUtilities/TestToolChain.cs b/Engine/Source/Programs/Shared/UnrealBuildTool.Tests/TestUtilities/TestToolChain.cs
deleted file mode 100644
diff --git a/Engine/Source/Programs/Shared/UnrealBuildTool.Tests/TokenReaderTests.cs b/Engine/Source/Programs/Shared/UnrealBuildTool.Tests/TokenReaderTests.cs
deleted file mode 100644
diff --git a/Engine/Source/Programs/Shared/UnrealBuildTool.Tests/UnrealBuildTool.Tests.csproj b/Engine/Source/Programs/Shared/UnrealBuildTool.Tests/UnrealBuildTool.Tests.csproj
deleted file mode 100644
diff --git a/Engine/Source/Programs/Shared/UnrealBuildTool.Tests/WarningLevelDefaultAttributeTests.cs b/Engine/Source/Programs/Shared/UnrealBuildTool.Tests/WarningLevelDefaultAttributeTests.cs
deleted file mode 100644
diff --git a/Engine/Source/Programs/UnrealBuildTool.Tests/.editorconfig b/Engine/Source/Programs/UnrealBuildTool.Tests/.editorconfig
new file mode 100644
--- /dev/null
+++ b/Engine/Source/Programs/UnrealBuildTool.Tests/.editorconfig
@@ -0,0 +1,1204 @@
+# NOTE: Requires **VS2019 16.3** or later
+
+# EditorConfig is awesome:http://EditorConfig.org
+
+# Dotnet code style settings:
+[*.cs]
+
+
+
+csharp_style_prefer_switch_expression = false
+csharp_style_implicit_object_creation_when_type_is_apparent = true:warning
+dotnet_remove_unnecessary_suppression_exclusions = all
+
+# Sort using and Import directives with System.* appearing first
+dotnet_sort_system_directives_first = true
+dotnet_separate_import_directive_groups = false
+
+# Avoid "this." and "Me." if not necessary
+dotnet_style_qualification_for_field = false:error
+dotnet_style_qualification_for_property = false:error
+dotnet_style_qualification_for_method = false:error
+dotnet_style_qualification_for_event = false:error
+
+# Use language keywords instead of framework type names for type references
+dotnet_style_predefined_type_for_locals_parameters_members = true:error
+dotnet_style_predefined_type_for_member_access = false:error
+
+# Suggest more modern language features when available
+dotnet_style_object_initializer = false:suggestion
+dotnet_style_collection_initializer = false:suggestion
+dotnet_style_coalesce_expression = true:suggestion
+dotnet_style_null_propagation = true:suggestion
+dotnet_style_explicit_tuple_names = true:suggestion
+
+# Whitespace options
+dotnet_style_allow_multiple_blank_lines_experimental = false
+dotnet_style_allow_statement_immediately_after_block_experimental = true
+
+# Non-private static fields are PascalCase
+dotnet_naming_rule.non_private_static_fields_should_be_pascal_case.severity = suggestion
+dotnet_naming_rule.non_private_static_fields_should_be_pascal_case.symbols = non_private_static_fields
+dotnet_naming_rule.non_private_static_fields_should_be_pascal_case.style = non_private_static_field_style
+
+dotnet_naming_symbols.non_private_static_fields.applicable_kinds = field
+dotnet_naming_symbols.non_private_static_fields.applicable_accessibilities = public, protected, internal, protected_internal, private_protected
+dotnet_naming_symbols.non_private_static_fields.required_modifiers = static
+
+dotnet_naming_style.non_private_static_field_style.capitalization = pascal_case
+
+# Non-private readonly fields are PascalCase
+dotnet_naming_rule.non_private_readonly_fields_should_be_pascal_case.severity = suggestion
+dotnet_naming_rule.non_private_readonly_fields_should_be_pascal_case.symbols = non_private_readonly_fields
+dotnet_naming_rule.non_private_readonly_fields_should_be_pascal_case.style = non_private_readonly_field_style
+
+dotnet_naming_symbols.non_private_readonly_fields.applicable_kinds = field
+dotnet_naming_symbols.non_private_readonly_fields.applicable_accessibilities = public, protected, internal, protected_internal, private_protected
+dotnet_naming_symbols.non_private_readonly_fields.required_modifiers = readonly
+
+dotnet_naming_style.non_private_readonly_field_style.capitalization = pascal_case
+
+# Constants are PascalCase
+dotnet_naming_rule.constants_should_be_pascal_case.severity = suggestion
+dotnet_naming_rule.constants_should_be_pascal_case.symbols = constants
+dotnet_naming_rule.constants_should_be_pascal_case.style = constant_style
+
+dotnet_naming_symbols.constants.applicable_kinds = field, local
+dotnet_naming_symbols.constants.required_modifiers = const
+
+dotnet_naming_style.constant_style.capitalization = pascal_case
+
+# Static fields are camelCase and start with s_
+dotnet_naming_rule.static_fields_should_be_camel_case.severity = suggestion
+dotnet_naming_rule.static_fields_should_be_camel_case.symbols = static_fields
+dotnet_naming_rule.static_fields_should_be_camel_case.style = static_field_style
+
+dotnet_naming_symbols.static_fields.applicable_kinds = field
+dotnet_naming_symbols.static_fields.required_modifiers = static
+
+dotnet_naming_style.static_field_style.capitalization = camel_case
+dotnet_naming_style.static_field_style.required_prefix = s_
+
+# Instance fields are camelCase and start with _
+dotnet_naming_rule.instance_fields_should_be_camel_case.severity = suggestion
+dotnet_naming_rule.instance_fields_should_be_camel_case.symbols = instance_fields
+dotnet_naming_rule.instance_fields_should_be_camel_case.style = instance_field_style
+
+dotnet_naming_symbols.instance_fields.applicable_kinds = field
+
+dotnet_naming_style.instance_field_style.capitalization = camel_case
+dotnet_naming_style.instance_field_style.required_prefix = _
+
+# Locals and parameters are camelCase
+dotnet_naming_rule.locals_should_be_camel_case.severity = suggestion
+dotnet_naming_rule.locals_should_be_camel_case.symbols = locals_and_parameters
+dotnet_naming_rule.locals_should_be_camel_case.style = camel_case_style
+
+dotnet_naming_symbols.locals_and_parameters.applicable_kinds = parameter, local
+
+dotnet_naming_style.camel_case_style.capitalization = camel_case
+
+# Local functions are PascalCase
+dotnet_naming_rule.local_functions_should_be_pascal_case.severity = suggestion
+dotnet_naming_rule.local_functions_should_be_pascal_case.symbols = local_functions
+dotnet_naming_rule.local_functions_should_be_pascal_case.style = local_function_style
+
+dotnet_naming_symbols.local_functions.applicable_kinds = local_function
+
+dotnet_naming_style.local_function_style.capitalization = pascal_case
+
+# By default, name items with PascalCase
+dotnet_naming_rule.members_should_be_pascal_case.severity = suggestion
+dotnet_naming_rule.members_should_be_pascal_case.symbols = all_members
+dotnet_naming_rule.members_should_be_pascal_case.style = pascal_case_style
+
+dotnet_naming_symbols.all_members.applicable_kinds = *
+
+dotnet_naming_style.pascal_case_style.capitalization = pascal_case
+
+# IDE0073: File header
+dotnet_diagnostic.IDE0073.severity = warning
+file_header_template = Copyright Epic Games, Inc. All Rights Reserved.
+
+# IDE0035: Remove unreachable code
+dotnet_diagnostic.IDE0035.severity = warning
+
+# IDE0036: Order modifiers
+dotnet_diagnostic.IDE0036.severity = warning
+
+# IDE0043: Format string contains invalid placeholder
+dotnet_diagnostic.IDE0043.severity = warning
+
+# IDE0044: Make field readonly
+dotnet_diagnostic.IDE0044.severity = warning
+
+# IDE0028: Use collection initializers
+dotnet_diagnostic.IDE0028.severity = warning
+
+# IDE0031: Use null propagation
+dotnet_diagnostic.IDE0031.severity = warning
+
+# IDE0034: Simplify 'default' expression
+dotnet_diagnostic.IDE0034.severity = warning
+
+# IDE0054: Use compound assignment
+dotnet_diagnostic.IDE0054.severity = warning
+
+# IDE0062: Make local function static
+dotnet_diagnostic.IDE0062.severity = warning
+
+# IDE0071: Simplify interpolation
+dotnet_diagnostic.IDE0071.severity = warning
+
+# IDE0074: Use compound assignment
+dotnet_diagnostic.IDE0074.severity = warning
+
+# IDE0075: Simplify conditional expression
+dotnet_diagnostic.IDE0075.severity = warning
+
+# IDE0280: Use 'nameof'
+dotnet_diagnostic.IDE0280.severity = warning
+
+# IDE0300: Use collection expression for array
+dotnet_diagnostic.IDE0300.severity = warning
+
+# IDE0301: Use collection expression for empty
+dotnet_diagnostic.IDE0301.severity = warning
+
+# IDE0305: Use collection expression for fluent
+dotnet_diagnostic.IDE0305.severity = warning
+
+# IDE0011: Add braces
+csharp_prefer_braces = when_multiline:warning
+# NOTE: We need the below severity entry for Add Braces due to [REDACTED_PROFILE_URL]/roslyn/issues/44201
+dotnet_diagnostic.IDE0011.severity = warning
+
+# IDE0040: Add accessibility modifiers
+dotnet_diagnostic.IDE0040.severity = warning
+
+# CONSIDER: Are IDE0051 and IDE0052 too noisy to be warnings for IDE editing scenarios? Should they be made build-only warnings?
+# IDE0051: Remove unused private member
+dotnet_diagnostic.IDE0051.severity = warning
+
+# IDE0052: Remove unread private member
+dotnet_diagnostic.IDE0052.severity = warning
+
+# IDE0059: Unnecessary assignment to a value
+dotnet_diagnostic.IDE0059.severity = warning
+
+# IDE0060: Remove unused parameter
+dotnet_diagnostic.IDE0060.severity = warning
+
+# CA1012: Abstract types should not have public constructors
+dotnet_diagnostic.CA1012.severity = warning
+
+# IDE0005: Using directive is unnecessary
+dotnet_diagnostic.IDE0005.severity = warning
+
+# IDE0270: Null check can be simplified (if null check)
+dotnet_diagnostic.IDE0270.severity = silent
+
+# IDE0090: 'new' expression can be simplified
+dotnet_diagnostic.IDE0090.severity = warning
+
+# IDE1006: Enforce naming conventions
+dotnet_diagnostic.IDE1006.severity = warning
+
+# dotnet_style_allow_multiple_blank_lines_experimental
+dotnet_diagnostic.IDE2000.severity = warning
+
+# csharp_style_allow_embedded_statements_on_same_line_experimental
+dotnet_diagnostic.IDE2001.severity = warning
+
+# csharp_style_allow_blank_lines_between_consecutive_braces_experimental
+dotnet_diagnostic.IDE2002.severity = warning
+
+# dotnet_style_allow_statement_immediately_after_block_experimental
+dotnet_diagnostic.IDE2003.severity = warning
+
+# csharp_style_allow_blank_line_after_colon_in_constructor_initializer_experimental
+dotnet_diagnostic.IDE2004.severity = warning
+
+dotnet_style_require_accessibility_modifiers = false
+
+# CSharp code style settings:
+[*.cs]
+csharp_prefer_simple_using_statement = false
+csharp_style_prefer_range_operator = false
+
+# Newline settings
+csharp_new_line_before_open_brace = all
+csharp_new_line_before_else = true
+csharp_new_line_before_catch = true
+csharp_new_line_before_finally = true
+csharp_new_line_before_members_in_object_initializers = true
+csharp_new_line_before_members_in_anonymous_types = true
+csharp_new_line_between_query_expression_clauses = true
+
+# Indentation preferences
+csharp_indent_block_contents = true
+csharp_indent_braces = false
+csharp_indent_case_contents = true
+csharp_indent_case_contents_when_block = true
+csharp_indent_switch_labels = true
+csharp_indent_labels = flush_left
+
+# Whitespace options
+csharp_style_allow_embedded_statements_on_same_line_experimental = false
+csharp_style_allow_blank_lines_between_consecutive_braces_experimental = false
+csharp_style_allow_blank_line_after_colon_in_constructor_initializer_experimental = false
+
+# Prefer "var" everywhere
+dotnet_diagnostic.IDE0007.severity = error
+dotnet_diagnostic.IDE0008.severity = error
+csharp_style_var_for_built_in_types = false:error
+csharp_style_var_when_type_is_apparent = false:error
+csharp_style_var_elsewhere = false:error
+
+# Prefer method-like constructs to have a block body
+csharp_style_expression_bodied_methods = false:none
+csharp_style_expression_bodied_constructors = false:none
+csharp_style_expression_bodied_operators = false:none
+
+# Prefer property-like constructs to have an expression-body
+csharp_style_expression_bodied_properties = true:error
+csharp_style_expression_bodied_indexers = true:error
+csharp_style_expression_bodied_accessors = true:error
+
+# Suggest more modern language features when available
+csharp_style_pattern_matching_over_is_with_cast_check = true:error
+csharp_style_pattern_matching_over_as_with_null_check = false:error
+csharp_style_inlined_variable_declaration = false:suggestion
+csharp_style_throw_expression = true:error
+csharp_style_conditional_delegate_call = true:suggestion
+
+# Spacing
+csharp_space_after_cast = false
+csharp_space_after_colon_in_inheritance_clause = true
+csharp_space_after_comma = true
+csharp_space_after_dot = false
+csharp_space_after_keywords_in_control_flow_statements = true
+csharp_space_after_semicolon_in_for_statement = true
+csharp_space_around_binary_operators = before_and_after
+csharp_space_around_declaration_statements = do_not_ignore
+csharp_space_before_colon_in_inheritance_clause = true
+csharp_space_before_comma = false
+csharp_space_before_dot = false
+csharp_space_before_open_square_brackets = false
+csharp_space_before_semicolon_in_for_statement = false
+csharp_space_between_empty_square_brackets = false
+csharp_space_between_method_call_empty_parameter_list_parentheses = false
+csharp_space_between_method_call_name_and_opening_parenthesis = false
+csharp_space_between_method_call_parameter_list_parentheses = false
+csharp_space_between_method_declaration_empty_parameter_list_parentheses = false
+csharp_space_between_method_declaration_name_and_open_parenthesis = false
+csharp_space_between_method_declaration_parameter_list_parentheses = false
+csharp_space_between_parentheses = false
+csharp_space_between_square_brackets = false
+
+# Blocks are allowed
+csharp_prefer_braces = true:silent
+csharp_preserve_single_line_blocks = true
+csharp_preserve_single_line_statements = true
+
+# Async methods should have "Async" suffix
+dotnet_naming_rule.async_methods_end_in_async.symbols = any_async_methods
+dotnet_naming_rule.async_methods_end_in_async.style = end_in_async
+dotnet_naming_rule.async_methods_end_in_async.severity = warning
+
+dotnet_naming_symbols.any_async_methods.applicable_kinds = method
+dotnet_naming_symbols.any_async_methods.applicable_accessibilities = *
+dotnet_naming_symbols.any_async_methods.required_modifiers = async
+
+dotnet_naming_style.end_in_async.required_prefix = 
+dotnet_naming_style.end_in_async.required_suffix = Async
+dotnet_naming_style.end_in_async.capitalization = pascal_case
+dotnet_naming_style.end_in_async.word_separator = 
+
+# All Rules Enabled as build warnings
+# Description: All Rules are enabled as build warnings. Rules with IsEnabledByDefault = false are force enabled as build warnings.
+
+# Code files
+[*.{cs,vb}]
+
+
+# CA1000: Do not declare static members on generic types
+dotnet_diagnostic.CA1000.severity = warning
+
+# CA1001: Types that own disposable fields should be disposable
+dotnet_diagnostic.CA1001.severity = warning
+
+# CA1002: Do not expose generic lists
+dotnet_diagnostic.CA1002.severity = none
+
+# CA1003: Use generic event handler instances
+dotnet_diagnostic.CA1003.severity = none
+
+# CA1005: Avoid excessive parameters on generic types
+dotnet_diagnostic.CA1005.severity = warning
+
+# CA1008: Enums should have zero value
+dotnet_diagnostic.CA1008.severity = warning
+
+# CA1010: Generic interface should also be implemented
+dotnet_diagnostic.CA1010.severity = warning
+
+# CA1012: Abstract types should not have public constructors
+dotnet_diagnostic.CA1012.severity = warning
+
+# CA1014: Mark assemblies with CLSCompliant
+dotnet_diagnostic.CA1014.severity = warning
+
+# CA1016: Mark assemblies with assembly version
+dotnet_diagnostic.CA1016.severity = warning
+
+# CA1017: Mark assemblies with ComVisible
+dotnet_diagnostic.CA1017.severity = warning
+
+# CA1018: Mark attributes with AttributeUsageAttribute
+dotnet_diagnostic.CA1018.severity = warning
+
+# CA1019: Define accessors for attribute arguments
+dotnet_diagnostic.CA1019.severity = warning
+
+# CA1021: Avoid out parameters
+dotnet_diagnostic.CA1021.severity = none
+
+# CA1024: Use properties where appropriate
+dotnet_diagnostic.CA1024.severity = none
+
+# CA1027: Mark enums with FlagsAttribute
+dotnet_diagnostic.CA1027.severity = warning
+
+# CA1028: Enum Storage should be Int32
+dotnet_diagnostic.CA1028.severity = warning
+
+# CA1030: Use events where appropriate
+dotnet_diagnostic.CA1030.severity = warning
+
+# CA1031: Do not catch general exception types
+dotnet_diagnostic.CA1031.severity = none
+
+# CA1032: Implement standard exception constructors
+dotnet_diagnostic.CA1032.severity = none
+
+# CA1033: Interface methods should be callable by child types
+dotnet_diagnostic.CA1033.severity = warning
+
+# CA1034: Nested types should not be visible
+dotnet_diagnostic.CA1034.severity = warning
+
+# CA1036: Override methods on comparable types
+dotnet_diagnostic.CA1036.severity = warning
+
+# CA1040: Avoid empty interfaces
+dotnet_diagnostic.CA1040.severity = warning
+
+# CA1041: Provide ObsoleteAttribute message
+dotnet_diagnostic.CA1041.severity = warning
+
+# CA1043: Use Integral Or String Argument For Indexers
+dotnet_diagnostic.CA1043.severity = warning
+
+# CA1044: Properties should not be write only
+dotnet_diagnostic.CA1044.severity = warning
+
+# CA1045: Do not pass types by reference
+dotnet_diagnostic.CA1045.severity = warning
+
+# CA1046: Do not overload equality operator on reference types
+dotnet_diagnostic.CA1046.severity = warning
+
+# CA1047: Do not declare protected member in sealed type
+dotnet_diagnostic.CA1047.severity = warning
+
+# CA1050: Declare types in namespaces
+dotnet_diagnostic.CA1050.severity = warning
+
+# CA1051: Do not declare visible instance fields
+dotnet_diagnostic.CA1051.severity = warning
+
+# CA1052: Static holder types should be Static or NotInheritable
+dotnet_diagnostic.CA1052.severity = warning
+
+# CA1054: URI-like parameters should not be strings
+dotnet_diagnostic.CA1054.severity = warning
+
+# CA1055: URI-like return values should not be strings
+dotnet_diagnostic.CA1055.severity = warning
+
+# CA1056: URI-like properties should not be strings
+dotnet_diagnostic.CA1056.severity = warning
+
+# CA1058: Types should not extend certain base types
+dotnet_diagnostic.CA1058.severity = warning
+
+# CA1060: Move pinvokes to native methods class
+dotnet_diagnostic.CA1060.severity = none
+
+# CA1061: Do not hide base class methods
+dotnet_diagnostic.CA1061.severity = warning
+
+# CA1062: Validate arguments of public methods
+dotnet_diagnostic.CA1062.severity = none
+
+# CA1063: Implement IDisposable Correctly
+dotnet_diagnostic.CA1063.severity = warning
+
+# CA1064: Exceptions should be public
+dotnet_diagnostic.CA1064.severity = warning
+
+# CA1065: Do not raise exceptions in unexpected locations
+dotnet_diagnostic.CA1065.severity = none
+
+# CA1066: Implement IEquatable when overriding Object.Equals
+dotnet_diagnostic.CA1066.severity = warning
+
+# CA1067: Override Object.Equals(object) when implementing IEquatable<T>
+dotnet_diagnostic.CA1067.severity = warning
+
+# CA1068: CancellationToken parameters must come last
+dotnet_diagnostic.CA1068.severity = warning
+
+# CA1069: Enums values should not be duplicated
+dotnet_diagnostic.CA1069.severity = warning
+
+# CA1070: Do not declare event fields as virtual
+dotnet_diagnostic.CA1070.severity = warning
+
+# CA1200: Avoid using cref tags with a prefix
+dotnet_diagnostic.CA1200.severity = warning
+
+# CA1303: Do not pass literals as localized parameters
+dotnet_diagnostic.CA1303.severity = none
+
+# CA1304: Specify CultureInfo
+dotnet_diagnostic.CA1304.severity = none
+
+# CA1305: Specify IFormatProvider
+dotnet_diagnostic.CA1305.severity = none
+
+# CA1307: Specify StringComparison for clarity
+dotnet_diagnostic.CA1307.severity = warning
+
+# CA1308: Normalize strings to uppercase
+dotnet_diagnostic.CA1308.severity = warning
+
+# CA1309: Use ordinal string comparison
+dotnet_diagnostic.CA1309.severity = warning
+
+# CA1310: Specify StringComparison for correctness
+dotnet_diagnostic.CA1310.severity = warning
+
+# CA1401: P/Invokes should not be visible
+dotnet_diagnostic.CA1401.severity = warning
+
+# CA1416: Validate platform compatibility
+dotnet_diagnostic.CA1416.severity = warning
+
+# CA1417: Do not use 'OutAttribute' on string parameters for P/Invokes
+dotnet_diagnostic.CA1417.severity = warning
+
+# CA1418: Use valid platform string
+dotnet_diagnostic.CA1418.severity = warning
+
+# CA1419: Provide a parameterless constructor that is as visible as the containing type for concrete types derived from 'System.Runtime.InteropServices.SafeHandle'
+dotnet_diagnostic.CA1419.severity = warning
+
+# CA1501: Avoid excessive inheritance
+dotnet_diagnostic.CA1501.severity = warning
+
+# CA1502: Avoid excessive complexity
+dotnet_diagnostic.CA1502.severity = none
+
+# CA1505: Avoid unmaintainable code
+dotnet_diagnostic.CA1505.severity = warning
+
+# CA1506: Avoid excessive class coupling
+dotnet_diagnostic.CA1506.severity = none
+
+# CA1507: Use nameof to express symbol names
+dotnet_diagnostic.CA1507.severity = warning
+
+# CA1508: Avoid dead conditional code
+dotnet_diagnostic.CA1508.severity = warning
+
+# CA1509: Invalid entry in code metrics rule specification file
+dotnet_diagnostic.CA1509.severity = warning
+
+# CA1700: Do not name enum values 'Reserved'
+dotnet_diagnostic.CA1700.severity = warning
+
+# CA1707: Identifiers should not contain underscores
+dotnet_diagnostic.CA1707.severity = warning
+
+# CA1708: Identifiers should differ by more than case
+dotnet_diagnostic.CA1708.severity = warning
+
+# CA1710: Identifiers should have correct suffix
+dotnet_diagnostic.CA1710.severity = warning
+
+# CA1711: Identifiers should not have incorrect suffix
+dotnet_diagnostic.CA1711.severity = none
+
+# CA1712: Do not prefix enum values with type name
+dotnet_diagnostic.CA1712.severity = warning
+
+# CA1713: Events should not have 'Before' or 'After' prefix
+dotnet_diagnostic.CA1713.severity = warning
+
+# CA1715: Identifiers should have correct prefix
+dotnet_diagnostic.CA1715.severity = warning
+
+# CA1716: Identifiers should not match keywords
+dotnet_diagnostic.CA1716.severity = warning
+
+# CA1720: Identifier contains type name
+dotnet_diagnostic.CA1720.severity = none
+
+# CA1721: Property names should not match get methods
+dotnet_diagnostic.CA1721.severity = warning
+
+# CA1724: Type names should not match namespaces
+dotnet_diagnostic.CA1724.severity = warning
+
+# CA1725: Parameter names should match base declaration
+dotnet_diagnostic.CA1725.severity = none
+
+# CA1727: Use PascalCase for named placeholders
+dotnet_diagnostic.CA1727.severity = warning
+
+# CA1802: Use literals where appropriate
+dotnet_diagnostic.CA1802.severity = warning
+
+# CA1805: Do not initialize unnecessarily
+dotnet_diagnostic.CA1805.severity = none
+
+# CA1806: Do not ignore method results
+dotnet_diagnostic.CA1806.severity = warning
+
+# CA1810: Initialize reference type static fields inline
+dotnet_diagnostic.CA1810.severity = warning
+
+# CA1812: Avoid uninstantiated internal classes
+dotnet_diagnostic.CA1812.severity = none
+
+# CA1813: Avoid unsealed attributes
+dotnet_diagnostic.CA1813.severity = warning
+
+# CA1814: Prefer jagged arrays over multidimensional
+dotnet_diagnostic.CA1814.severity = warning
+
+# CA1815: Override equals and operator equals on value types
+dotnet_diagnostic.CA1815.severity = none
+
+# CA1816: Dispose methods should call SuppressFinalize
+dotnet_diagnostic.CA1816.severity = warning
+
+# CA1819: Properties should not return arrays (rationale: System.Text.Json serialization requires array properties)
+dotnet_diagnostic.CA1819.severity = silent
+
+# CA1820: Test for empty strings using string length
+dotnet_diagnostic.CA1820.severity = warning
+
+# CA1821: Remove empty Finalizers
+dotnet_diagnostic.CA1821.severity = warning
+
+# CA1822: Mark members as static
+dotnet_diagnostic.CA1822.severity = warning
+
+# CA1823: Avoid unused private fields
+dotnet_diagnostic.CA1823.severity = warning
+
+# CA1824: Mark assemblies with NeutralResourcesLanguageAttribute
+dotnet_diagnostic.CA1824.severity = warning
+
+# CA1825: Avoid zero-length array allocations
+dotnet_diagnostic.CA1825.severity = warning
+
+# CA1826: Do not use Enumerable methods on indexable collections
+dotnet_diagnostic.CA1826.severity = warning
+
+# CA1827: Do not use Count() or LongCount() when Any() can be used
+dotnet_diagnostic.CA1827.severity = warning
+
+# CA1828: Do not use CountAsync() or LongCountAsync() when AnyAsync() can be used
+dotnet_diagnostic.CA1828.severity = warning
+
+# CA1829: Use Length/Count property instead of Count() when available
+dotnet_diagnostic.CA1829.severity = warning
+
+# CA1830: Prefer strongly-typed Append and Insert method overloads on StringBuilder
+dotnet_diagnostic.CA1830.severity = warning
+
+# CA1831: Use AsSpan or AsMemory instead of Range-based indexers when appropriate
+dotnet_diagnostic.CA1831.severity = warning
+
+# CA1832: Use AsSpan or AsMemory instead of Range-based indexers when appropriate
+dotnet_diagnostic.CA1832.severity = warning
+
+# CA1833: Use AsSpan or AsMemory instead of Range-based indexers when appropriate
+dotnet_diagnostic.CA1833.severity = warning
+
+# CA1834: Consider using 'StringBuilder.Append(char)' when applicable
+dotnet_diagnostic.CA1834.severity = warning
+
+# CA1835: Prefer the 'Memory'-based overloads for 'ReadAsync' and 'WriteAsync'
+dotnet_diagnostic.CA1835.severity = none
+
+# CA1836: Prefer IsEmpty over Count
+dotnet_diagnostic.CA1836.severity = warning
+
+# CA1837: Use 'Environment.ProcessId'
+dotnet_diagnostic.CA1837.severity = warning
+
+# CA1838: Avoid 'StringBuilder' parameters for P/Invokes
+dotnet_diagnostic.CA1838.severity = warning
+
+# CA1839: Use 'Environment.ProcessPath'
+dotnet_diagnostic.CA1839.severity = warning
+
+# CA1840: Use 'Environment.CurrentManagedThreadId'
+dotnet_diagnostic.CA1840.severity = warning
+
+# CA1841: Prefer Dictionary.Contains methods
+dotnet_diagnostic.CA1841.severity = warning
+
+# CA1842: Do not use 'WhenAll' with a single task
+dotnet_diagnostic.CA1842.severity = warning
+
+# CA1843: Do not use 'WaitAll' with a single task
+dotnet_diagnostic.CA1843.severity = warning
+
+# CA1844: Provide memory-based overrides of async methods when subclassing 'Stream'
+dotnet_diagnostic.CA1844.severity = warning
+
+# CA1845: Use span-based 'string.Concat'
+dotnet_diagnostic.CA1845.severity = none
+
+# CA1846: Prefer 'AsSpan' over 'Substring'
+dotnet_diagnostic.CA1846.severity = none
+
+# CA1847: Use char literal for a single character lookup
+dotnet_diagnostic.CA1847.severity = warning
+
+# CA1848: Use the LoggerMessage delegates
+dotnet_diagnostic.CA1848.severity = none
+
+# CA1849: Call async methods when in an async method
+dotnet_diagnostic.CA1849.severity = warning
+
+# CA2000: Dispose objects before losing scope
+dotnet_diagnostic.CA2000.severity = warning
+
+# CA2002: Do not lock on objects with weak identity
+dotnet_diagnostic.CA2002.severity = warning
+
+# CA2007: Consider calling ConfigureAwait on the awaited task
+dotnet_diagnostic.CA2007.severity = none
+
+# CA2008: Do not create tasks without passing a TaskScheduler
+dotnet_diagnostic.CA2008.severity = warning
+
+# CA2009: Do not call ToImmutableCollection on an ImmutableCollection value
+dotnet_diagnostic.CA2009.severity = warning
+
+# CA2011: Avoid infinite recursion
+dotnet_diagnostic.CA2011.severity = warning
+
+# CA2012: Use ValueTasks correctly
+dotnet_diagnostic.CA2012.severity = warning
+
+# CA2013: Do not use ReferenceEquals with value types
+dotnet_diagnostic.CA2013.severity = warning
+
+# CA2014: Do not use stackalloc in loops
+dotnet_diagnostic.CA2014.severity = warning
+
+# CA2015: Do not define finalizers for types derived from MemoryManager<T>
+dotnet_diagnostic.CA2015.severity = warning
+
+# CA2016: Forward the 'CancellationToken' parameter to methods
+dotnet_diagnostic.CA2016.severity = warning
+
+# CA2017: Parameter count mismatch
+dotnet_diagnostic.CA2017.severity = warning
+
+# CA2018: 'Buffer.BlockCopy' expects the number of bytes to be copied for the 'count' argument
+dotnet_diagnostic.CA2018.severity = warning
+
+# CA2100: Review SQL queries for security vulnerabilities
+dotnet_diagnostic.CA2100.severity = warning
+
+# CA2101: Specify marshaling for P/Invoke string arguments
+dotnet_diagnostic.CA2101.severity = none
+
+# CA2109: Review visible event handlers
+dotnet_diagnostic.CA2109.severity = warning
+
+# CA2119: Seal methods that satisfy private interfaces
+dotnet_diagnostic.CA2119.severity = warning
+
+# CA2153: Do Not Catch Corrupted State Exceptions
+dotnet_diagnostic.CA2153.severity = warning
+
+# CA2200: Rethrow to preserve stack details
+dotnet_diagnostic.CA2200.severity = warning
+
+# CA2201: Do not raise reserved exception types
+dotnet_diagnostic.CA2201.severity = none
+
+# CA2207: Initialize value type static fields inline
+dotnet_diagnostic.CA2207.severity = warning
+
+# CA2208: Instantiate argument exceptions correctly
+dotnet_diagnostic.CA2208.severity = warning
+
+# CA2211: Non-constant fields should not be visible
+dotnet_diagnostic.CA2211.severity = warning
+
+# CA2213: Disposable fields should be disposed
+dotnet_diagnostic.CA2213.severity = warning
+
+# CA2214: Do not call overridable methods in constructors
+dotnet_diagnostic.CA2214.severity = warning
+
+# CA2215: Dispose methods should call base class dispose
+dotnet_diagnostic.CA2215.severity = warning
+
+# CA2216: Disposable types should declare finalizer
+dotnet_diagnostic.CA2216.severity = warning
+
+# CA2217: Do not mark enums with FlagsAttribute
+dotnet_diagnostic.CA2217.severity = warning
+
+# CA2218: Override GetHashCode on overriding Equals
+dotnet_diagnostic.CA2218.severity = warning
+
+# CA2219: Do not raise exceptions in finally clauses
+dotnet_diagnostic.CA2219.severity = warning
+
+# CA2224: Override Equals on overloading operator equals
+dotnet_diagnostic.CA2224.severity = warning
+
+# CA2225: Operator overloads have named alternates
+dotnet_diagnostic.CA2225.severity = none
+
+# CA2226: Operators should have symmetrical overloads
+dotnet_diagnostic.CA2226.severity = warning
+
+# CA2227: Collection properties should be read only
+dotnet_diagnostic.CA2227.severity = warning
+
+# CA2229: Implement serialization constructors
+dotnet_diagnostic.CA2229.severity = warning
+
+# CA2231: Overload operator equals on overriding value type Equals
+dotnet_diagnostic.CA2231.severity = warning
+
+# CA2234: Pass system uri objects instead of strings
+dotnet_diagnostic.CA2234.severity = warning
+
+# CA2235: Mark all non-serializable fields
+dotnet_diagnostic.CA2235.severity = warning
+
+# CA2237: Mark ISerializable types with serializable
+dotnet_diagnostic.CA2237.severity = warning
+
+# CA2241: Provide correct arguments to formatting methods
+dotnet_diagnostic.CA2241.severity = warning
+
+# CA2242: Test for NaN correctly
+dotnet_diagnostic.CA2242.severity = warning
+
+# CA2243: Attribute string literals should parse correctly
+dotnet_diagnostic.CA2243.severity = warning
+
+# CA2244: Do not duplicate indexed element initializations
+dotnet_diagnostic.CA2244.severity = warning
+
+# CA2245: Do not assign a property to itself
+dotnet_diagnostic.CA2245.severity = warning
+
+# CA2246: Assigning symbol and its member in the same statement
+dotnet_diagnostic.CA2246.severity = warning
+
+# CA2247: Argument passed to TaskCompletionSource constructor should be TaskCreationOptions enum instead of TaskContinuationOptions enum
+dotnet_diagnostic.CA2247.severity = warning
+
+# CA2248: Provide correct 'enum' argument to 'Enum.HasFlag'
+dotnet_diagnostic.CA2248.severity = warning
+
+# CA2249: Consider using 'string.Contains' instead of 'string.IndexOf'
+dotnet_diagnostic.CA2249.severity = warning
+
+# CA2250: Use 'ThrowIfCancellationRequested'
+dotnet_diagnostic.CA2250.severity = warning
+
+# CA2251: Use 'string.Equals'
+dotnet_diagnostic.CA2251.severity = warning
+
+# CA2252: This API requires opting into preview features
+dotnet_diagnostic.CA2252.severity = warning
+
+# CA2253: Named placeholders should not be numeric values
+dotnet_diagnostic.CA2253.severity = warning
+
+# CA2254: Template should be a static expression
+dotnet_diagnostic.CA2254.severity = warning
+
+# CA2255: The 'ModuleInitializer' attribute should not be used in libraries
+dotnet_diagnostic.CA2255.severity = warning
+
+# CA2256: All members declared in parent interfaces must have an implementation in a DynamicInterfaceCastableImplementation-attributed interface
+dotnet_diagnostic.CA2256.severity = warning
+
+# CA2257: Members defined on an interface with the 'DynamicInterfaceCastableImplementationAttribute' should be 'static'
+dotnet_diagnostic.CA2257.severity = warning
+
+# CA2258: Providing a 'DynamicInterfaceCastableImplementation' interface in Visual Basic is unsupported
+dotnet_diagnostic.CA2258.severity = warning
+
+# CA2300: Do not use insecure deserializer BinaryFormatter
+dotnet_diagnostic.CA2300.severity = warning
+
+# CA2301: Do not call BinaryFormatter.Deserialize without first setting BinaryFormatter.Binder
+dotnet_diagnostic.CA2301.severity = warning
+
+# CA2302: Ensure BinaryFormatter.Binder is set before calling BinaryFormatter.Deserialize
+dotnet_diagnostic.CA2302.severity = warning
+
+# CA2305: Do not use insecure deserializer LosFormatter
+dotnet_diagnostic.CA2305.severity = warning
+
+# CA2310: Do not use insecure deserializer NetDataContractSerializer
+dotnet_diagnostic.CA2310.severity = warning
+
+# CA2311: Do not deserialize without first setting NetDataContractSerializer.Binder
+dotnet_diagnostic.CA2311.severity = warning
+
+# CA2312: Ensure NetDataContractSerializer.Binder is set before deserializing
+dotnet_diagnostic.CA2312.severity = warning
+
+# CA2315: Do not use insecure deserializer ObjectStateFormatter
+dotnet_diagnostic.CA2315.severity = warning
+
+# CA2321: Do not deserialize with JavaScriptSerializer using a SimpleTypeResolver
+dotnet_diagnostic.CA2321.severity = warning
+
+# CA2322: Ensure JavaScriptSerializer is not initialized with SimpleTypeResolver before deserializing
+dotnet_diagnostic.CA2322.severity = warning
+
+# CA2326: Do not use TypeNameHandling values other than None
+dotnet_diagnostic.CA2326.severity = warning
+
+# CA2327: Do not use insecure JsonSerializerSettings
+dotnet_diagnostic.CA2327.severity = warning
+
+# CA2328: Ensure that JsonSerializerSettings are secure
+dotnet_diagnostic.CA2328.severity = warning
+
+# CA2329: Do not deserialize with JsonSerializer using an insecure configuration
+dotnet_diagnostic.CA2329.severity = warning
+
+# CA2330: Ensure that JsonSerializer has a secure configuration when deserializing
+dotnet_diagnostic.CA2330.severity = warning
+
+# CA2350: Do not use DataTable.ReadXml() with untrusted data
+dotnet_diagnostic.CA2350.severity = warning
+
+# CA2351: Do not use DataSet.ReadXml() with untrusted data
+dotnet_diagnostic.CA2351.severity = warning
+
+# CA2352: Unsafe DataSet or DataTable in serializable type can be vulnerable to remote code execution attacks
+dotnet_diagnostic.CA2352.severity = warning
+
+# CA2353: Unsafe DataSet or DataTable in serializable type
+dotnet_diagnostic.CA2353.severity = warning
+
+# CA2354: Unsafe DataSet or DataTable in deserialized object graph can be vulnerable to remote code execution attacks
+dotnet_diagnostic.CA2354.severity = warning
+
+# CA2355: Unsafe DataSet or DataTable type found in deserializable object graph
+dotnet_diagnostic.CA2355.severity = warning
+
+# CA2356: Unsafe DataSet or DataTable type in web deserializable object graph
+dotnet_diagnostic.CA2356.severity = warning
+
+# CA2361: Ensure auto-generated class containing DataSet.ReadXml() is not used with untrusted data
+dotnet_diagnostic.CA2361.severity = warning
+
+# CA2362: Unsafe DataSet or DataTable in auto-generated serializable type can be vulnerable to remote code execution attacks
+dotnet_diagnostic.CA2362.severity = warning
+
+# CA3001: Review code for SQL injection vulnerabilities
+dotnet_diagnostic.CA3001.severity = warning
+
+# CA3002: Review code for XSS vulnerabilities
+dotnet_diagnostic.CA3002.severity = warning
+
+# CA3003: Review code for file path injection vulnerabilities
+dotnet_diagnostic.CA3003.severity = warning
+
+# CA3004: Review code for information disclosure vulnerabilities
+dotnet_diagnostic.CA3004.severity = warning
+
+# CA3005: Review code for LDAP injection vulnerabilities
+dotnet_diagnostic.CA3005.severity = warning
+
+# CA3006: Review code for process command injection vulnerabilities
+dotnet_diagnostic.CA3006.severity = warning
+
+# CA3007: Review code for open redirect vulnerabilities
+dotnet_diagnostic.CA3007.severity = warning
+
+# CA3008: Review code for XPath injection vulnerabilities
+dotnet_diagnostic.CA3008.severity = warning
+
+# CA3009: Review code for XML injection vulnerabilities
+dotnet_diagnostic.CA3009.severity = warning
+
+# CA3010: Review code for XAML injection vulnerabilities
+dotnet_diagnostic.CA3010.severity = warning
+
+# CA3011: Review code for DLL injection vulnerabilities
+dotnet_diagnostic.CA3011.severity = warning
+
+# CA3012: Review code for regex injection vulnerabilities
+dotnet_diagnostic.CA3012.severity = warning
+
+# CA3061: Do Not Add Schema By URL
+dotnet_diagnostic.CA3061.severity = warning
+
+# CA3075: Insecure DTD processing in XML
+dotnet_diagnostic.CA3075.severity = warning
+
+# CA3076: Insecure XSLT script processing.
+dotnet_diagnostic.CA3076.severity = warning
+
+# CA3077: Insecure Processing in API Design, XmlDocument and XmlTextReader
+dotnet_diagnostic.CA3077.severity = warning
+
+# CA3147: Mark Verb Handlers With Validate Antiforgery Token
+dotnet_diagnostic.CA3147.severity = warning
+
+# CA5350: Do Not Use Weak Cryptographic Algorithms
+dotnet_diagnostic.CA5350.severity = warning
+
+# CA5351: Do Not Use Broken Cryptographic Algorithms
+dotnet_diagnostic.CA5351.severity = warning
+
+# CA5358: Review cipher mode usage with cryptography experts
+dotnet_diagnostic.CA5358.severity = warning
+
+# CA5359: Do Not Disable Certificate Validation
+dotnet_diagnostic.CA5359.severity = warning
+
+# CA5360: Do Not Call Dangerous Methods In Deserialization
+dotnet_diagnostic.CA5360.severity = warning
+
+# CA5361: Do Not Disable SChannel Use of Strong Crypto
+dotnet_diagnostic.CA5361.severity = warning
+
+# CA5362: Potential reference cycle in deserialized object graph
+dotnet_diagnostic.CA5362.severity = warning
+
+# CA5363: Do Not Disable Request Validation
+dotnet_diagnostic.CA5363.severity = warning
+
+# CA5364: Do Not Use Deprecated Security Protocols
+dotnet_diagnostic.CA5364.severity = warning
+
+# CA5365: Do Not Disable HTTP Header Checking
+dotnet_diagnostic.CA5365.severity = warning
+
+# CA5366: Use XmlReader for 'DataSet.ReadXml()'
+dotnet_diagnostic.CA5366.severity = warning
+
+# CA5367: Do Not Serialize Types With Pointer Fields
+dotnet_diagnostic.CA5367.severity = warning
+
+# CA5368: Set ViewStateUserKey For Classes Derived From Page
+dotnet_diagnostic.CA5368.severity = warning
+
+# CA5369: Use XmlReader for 'XmlSerializer.Deserialize()'
+dotnet_diagnostic.CA5369.severity = warning
+
+# CA5370: Use XmlReader for XmlValidatingReader constructor
+dotnet_diagnostic.CA5370.severity = warning
+
+# CA5371: Use XmlReader for 'XmlSchema.Read()'
+dotnet_diagnostic.CA5371.severity = warning
+
+# CA5372: Use XmlReader for XPathDocument constructor
+dotnet_diagnostic.CA5372.severity = warning
+
+# CA5373: Do not use obsolete key derivation function
+dotnet_diagnostic.CA5373.severity = warning
+
+# CA5374: Do Not Use XslTransform
+dotnet_diagnostic.CA5374.severity = warning
+
+# CA5375: Do Not Use Account Shared Access Signature
+dotnet_diagnostic.CA5375.severity = warning
+
+# CA5376: Use SharedAccessProtocol HttpsOnly
+dotnet_diagnostic.CA5376.severity = warning
+
+# CA5377: Use Container Level Access Policy
+dotnet_diagnostic.CA5377.severity = warning
+
+# CA5378: Do not disable ServicePointManagerSecurityProtocols
+dotnet_diagnostic.CA5378.severity = warning
+
+# CA5379: Ensure Key Derivation Function algorithm is sufficiently strong
+dotnet_diagnostic.CA5379.severity = warning
+
+# CA5380: Do Not Add Certificates To Root Store
+dotnet_diagnostic.CA5380.severity = warning
+
+# CA5381: Ensure Certificates Are Not Added To Root Store
+dotnet_diagnostic.CA5381.severity = warning
+
+# CA5382: Use Secure Cookies In ASP.NET Core
+dotnet_diagnostic.CA5382.severity = warning
+
+# CA5383: Ensure Use Secure Cookies In ASP.NET Core
+dotnet_diagnostic.CA5383.severity = warning
+
+# CA5384: Do Not Use Digital Signature Algorithm (DSA)
+dotnet_diagnostic.CA5384.severity = warning
+
+# CA5385: Use Rivest–Shamir–Adleman (RSA) Algorithm With Sufficient Key Size
+dotnet_diagnostic.CA5385.severity = warning
+
+# CA5386: Avoid hardcoding SecurityProtocolType value
+dotnet_diagnostic.CA5386.severity = warning
+
+# CA5387: Do Not Use Weak Key Derivation Function With Insufficient Iteration Count
+dotnet_diagnostic.CA5387.severity = warning
+
+# CA5388: Ensure Sufficient Iteration Count When Using Weak Key Derivation Function
+dotnet_diagnostic.CA5388.severity = warning
+
+# CA5389: Do Not Add Archive Item's Path To The Target File System Path
+dotnet_diagnostic.CA5389.severity = warning
+
+# CA5390: Do not hard-code encryption key (disabled: expensive interprocedural dataflow analyzer)
+dotnet_diagnostic.CA5390.severity = none
+
+# CA5391: Use antiforgery tokens in ASP.NET Core MVC controllers
+dotnet_diagnostic.CA5391.severity = warning
+
+# CA5392: Use DefaultDllImportSearchPaths attribute for P/Invokes
+dotnet_diagnostic.CA5392.severity = none
+
+# CA5393: Do not use unsafe DllImportSearchPath value
+dotnet_diagnostic.CA5393.severity = warning
+
+# CA5394: Do not use insecure randomness
+dotnet_diagnostic.CA5394.severity = warning
+
+# CA5395: Miss HttpVerb attribute for action methods
+dotnet_diagnostic.CA5395.severity = warning
+
+# CA5396: Set HttpOnly to true for HttpCookie
+dotnet_diagnostic.CA5396.severity = warning
+
+# CA5397: Do not use deprecated SslProtocols values
+dotnet_diagnostic.CA5397.severity = warning
+
+# CA5398: Avoid hardcoded SslProtocols values
+dotnet_diagnostic.CA5398.severity = warning
+
+# CA5399: HttpClients should enable certificate revocation list checks
+dotnet_diagnostic.CA5399.severity = warning
+
+# CA5400: Ensure HttpClient certificate revocation list check is not disabled
+dotnet_diagnostic.CA5400.severity = warning
+
+# CA5401: Do not use CreateEncryptor with non-default IV
+dotnet_diagnostic.CA5401.severity = warning
+
+# CA5402: Use CreateEncryptor with the default IV 
+dotnet_diagnostic.CA5402.severity = warning
+
+# CA5403: Do not hard-code certificate (disabled: expensive interprocedural dataflow analyzer)
+dotnet_diagnostic.CA5403.severity = none
+
+# CA5404: Do not disable token validation checks
+dotnet_diagnostic.CA5404.severity = none
+
+# CA5405: Do not always skip token validation in delegates
+dotnet_diagnostic.CA5405.severity = warning
+
+# VSTHRD003: Avoid awaiting or returning a Task representing work that was not started within your context as that can lead to deadlocks.
+dotnet_diagnostic.VSTHRD003.severity = none
+
+# MA0026: TODO comment
+dotnet_diagnostic.MA0026.severity = none
+
+# MA0048: File name must match type name
+dotnet_diagnostic.MA0048.severity = none
+
+# MA0004: Use Task.ConfigureAwait(false) if the current SynchronizationContext is not needed
+dotnet_diagnostic.MA0004.severity = none
+
+# MA0051: Method is too long
+dotnet_diagnostic.MA0051.severity = none
+
+# MA0002: Use an overload that has a IEqualityComparer<string> or IComparer<string> parameter
+dotnet_diagnostic.MA0002.severity = none
+
+# MA0016: Prefer using collection abstraction instead of implementation
+dotnet_diagnostic.MA0016.severity = none
+
+# MA0006: Use string.Equals instead of Equals operator
+dotnet_diagnostic.MA0006.severity = none
+
+# MA0061: Method overrides should not change default values
+dotnet_diagnostic.MA0061.severity = none
+
+# MA0134: Observe result of async calls
+dotnet_diagnostic.MA0134.severity = none
+
+# MA0011: Use an overload of 'ToString' that has a 'System.IFormatProvider' parameter
+dotnet_diagnostic.MA0011.severity = none
+
+# MA0009: Regular expressions should not be vulnerable to Denial of Service attacks
+dotnet_diagnostic.MA0009.severity = none
+
+# MA0099: Use Explicit enum value for 'SlackNotificationSink.ReactionFlags' instead of 0
+dotnet_diagnostic.MA0099.severity = none
+
+# MA0025: Implement the functionality (or raise NotSupportedException or PlatformNotSupportedException)
+dotnet_diagnostic.MA0025.severity = none
+
+# MA0015: Use an overload of 'System.ArgumentException' with the parameter name
+dotnet_diagnostic.MA0015.severity = none
+
+# MA0132: Do not convert implicitly to DateTimeOffset
+dotnet_diagnostic.MA0132.severity = none
+
+# MA0023: Add RegexOptions.ExplicitCapture to prevent capturing unneeded groups
+dotnet_diagnostic.MA0023.severity = none
+
+# MA0056: Do not call overridable members in constructor
+dotnet_diagnostic.MA0056.severity = none
+
+# MA0144: Use System.OperatingSystem to check the current OS
+dotnet_diagnostic.MA0144.severity = none
+
+# MA0046 - Use EventHandler<T> to declare events
+dotnet_diagnostic.MA0046.severity = none
+
+# MA0008: Add StructLayoutAttribute
+dotnet_diagnostic.MA0008.severity = none
+
+[*Test*.cs]
+
+# CA1707: Identifiers should not contain underscores
+dotnet_diagnostic.CA1707.severity = none
+
+# CA5394: Do not use insecure randomness
+dotnet_diagnostic.CA5394.severity = none
diff --git a/Engine/Source/Programs/UnrealBuildTool.Tests/ConfigFileTests.cs b/Engine/Source/Programs/UnrealBuildTool.Tests/ConfigFileTests.cs
new file mode 100644
--- /dev/null
+++ b/Engine/Source/Programs/UnrealBuildTool.Tests/ConfigFileTests.cs
@@ -0,0 +1,79 @@
+// Copyright Epic Games, Inc. All Rights Reserved.
+
+using System;
+using System.IO;
+using EpicGames.Core;
+using Microsoft.VisualStudio.TestTools.UnitTesting;
+
+namespace UnrealBuildTool.Tests
+{
+	[TestClass]
+	public class ConfigFileTests
+	{
+		#region -- Test Context --
+
+		private const string TEST_FILE_NAME = "ConfigFileTest.ini";
+
+		private class TestConfigFile : IDisposable
+		{
+			internal string? TemporaryFile { get; private set; }
+
+			internal TestConfigFile(string fileName, string contents)
+			{
+				try
+				{
+					TemporaryFile = Path.Combine(System.IO.Path.GetTempPath(), Guid.NewGuid().ToString() + "_" + fileName);
+					using (StreamWriter writer = new(TemporaryFile))
+					{
+						writer.Write(contents);
+					}
+				}
+				catch (Exception)
+				{
+					TemporaryFile = null;
+				}
+			}
+
+			public void Dispose()
+			{
+				if (!String.IsNullOrEmpty(TemporaryFile) && File.Exists(TemporaryFile))
+				{
+					File.Delete(TemporaryFile);
+				}
+			}
+		}
+
+		#endregion
+
+		[TestMethod]
+		public void ConfigFileTestsQuotations()
+		{
+			const string TestSection = "MySection";
+			const string TestValueKey = "Name";
+			const string TestValue = "\"https://www.unrealengine.com/\"";
+
+			string testContents = String.Format("[{0}]\r\n{1}={2}", TestSection, TestValueKey, TestValue);
+
+			using (TestConfigFile temporaryConfigFile = new(TEST_FILE_NAME, testContents))
+			{
+				string? testFilePath = temporaryConfigFile.TemporaryFile;
+
+				Assert.IsFalse(String.IsNullOrEmpty(testFilePath), "Unable to generate a test file.");
+
+				ConfigFile configFile;
+				FileReference configFileLocation = new(testFilePath);
+
+				Assert.IsTrue(FileReference.Exists(configFileLocation), "Unable to acquire a test file.");
+				
+				configFile = new ConfigFile(configFileLocation);
+
+				FileReference.MakeWriteable(configFileLocation);
+				configFile.Write(configFileLocation);
+
+				Assert.IsNotNull(temporaryConfigFile.TemporaryFile);
+				string fileContents = File.ReadAllText(temporaryConfigFile.TemporaryFile);
+				Assert.AreEqual(testContents.Trim(), fileContents.Trim(), "The file contents after writing do not match the expected contents.");
+			}
+		}
+	}
+}
diff --git a/Engine/Source/Programs/UnrealBuildTool.Tests/CppCompileWarningsTests.cs b/Engine/Source/Programs/UnrealBuildTool.Tests/CppCompileWarningsTests.cs
new file mode 100644
--- /dev/null
+++ b/Engine/Source/Programs/UnrealBuildTool.Tests/CppCompileWarningsTests.cs
@@ -0,0 +1,803 @@
+// Copyright Epic Games, Inc. All Rights Reserved.
+
+using System;
+using System.Collections.Generic;
+using System.IO;
+using System.Linq;
+using EpicGames.Core;
+using Microsoft.VisualStudio.TestTools.UnitTesting;
+using UnrealBuildTool.Tests.TestUtilities;
+
+namespace UnrealBuildTool.Tests
+{
+	[TestClass]
+	public class CppCompileWarningsTests
+	{
+		internal class TestBuildSettingsProvider : IBuildContextProvider
+		{
+			internal BuildSettingsVersion _testBuildSettings = BuildSettingsVersion.V1;
+
+			public BuildSettingsVersion GetBuildSettings()
+			{
+				return _testBuildSettings;
+			}
+		}
+
+		[TestMethod]
+		[DataRow(WarningLevel.Default, null, true, WarningLevel.Default)]                   // Null applied, no overwrite
+		[DataRow(WarningLevel.Default, WarningLevel.Off, true, WarningLevel.Off)]           // Overwrite with Off
+		[DataRow(WarningLevel.Default, WarningLevel.Warning, true, WarningLevel.Warning)]   // Overwrite with Warning
+		[DataRow(WarningLevel.Default, WarningLevel.Error, true, WarningLevel.Error)]       // Overwrite with Error
+		[DataRow(WarningLevel.Default, WarningLevel.Default, true, WarningLevel.Default)]   // Applied is Default, no change
+
+		[DataRow(WarningLevel.Default, WarningLevel.Warning, false, WarningLevel.Warning)]  // Default, non-overwrite, change
+		[DataRow(WarningLevel.Default, WarningLevel.Error, false, WarningLevel.Error)]      // Default, non-overwrite, change
+
+		[DataRow(WarningLevel.Off, null, true, WarningLevel.Off)]                           // Null applied, no change
+		[DataRow(WarningLevel.Off, WarningLevel.Default, true, WarningLevel.Off)]           // Applied Default, no change
+		[DataRow(WarningLevel.Off, WarningLevel.Warning, true, WarningLevel.Warning)]       // Overwrite with Warning
+		[DataRow(WarningLevel.Off, WarningLevel.Error, true, WarningLevel.Error)]           // Overwrite with Error
+		[DataRow(WarningLevel.Off, WarningLevel.Warning, false, WarningLevel.Off)]          // No overwrite, no change
+
+		[DataRow(WarningLevel.Warning, WarningLevel.Error, true, WarningLevel.Error)]       // Overwrite with Error
+		[DataRow(WarningLevel.Warning, WarningLevel.Off, true, WarningLevel.Off)]           // Overwrite with Off
+		[DataRow(WarningLevel.Warning, WarningLevel.Off, false, WarningLevel.Warning)]      // No overwrite, remains Warning
+
+		[DataRow(WarningLevel.Error, WarningLevel.Off, true, WarningLevel.Off)]             // Overwrite with Off
+		[DataRow(WarningLevel.Error, WarningLevel.Warning, true, WarningLevel.Warning)]     // Overwrite with Warning
+		[DataRow(WarningLevel.Error, WarningLevel.Off, false, WarningLevel.Error)]          // No overwrite, remains Error
+		public void TestApplyWarnings(WarningLevel inSourceWarningLevel, WarningLevel? appliedWarningLevel, bool overwriteSourceValue, WarningLevel expected)
+		{
+			WarningLevel sourceWarningLevel = inSourceWarningLevel;
+			sourceWarningLevel.ApplyWarning(appliedWarningLevel, overwriteSourceValue);
+			Assert.AreEqual(expected, sourceWarningLevel);
+		}
+
+		[TestMethod]
+		public void TestDefaultBuildSettings()
+		{
+			CppCompileWarnings defaultBuildSettings = new();
+
+			Assert.AreEqual(WarningLevel.Warning, defaultBuildSettings.ShadowVariableWarningLevel);
+			Assert.AreEqual(WarningLevel.Off, defaultBuildSettings.UndefinedIdentifierWarningLevel);
+			Assert.AreEqual(WarningLevel.Off, defaultBuildSettings.UnsafeTypeCastWarningLevel);
+			Assert.AreEqual(WarningLevel.Off, defaultBuildSettings.SwitchUnhandledEnumeratorWarningLevel);
+		}
+
+		[TestMethod]
+		public void TestClonedMutableReturn()
+		{
+			CppCompileWarnings defaultBuildSettings = new();
+			defaultBuildSettings.ShadowVariableWarningLevel = WarningLevel.Error;
+			defaultBuildSettings.UndefinedIdentifierWarningLevel = WarningLevel.Warning;
+			defaultBuildSettings.UnsafeTypeCastWarningLevel = WarningLevel.Error;
+			defaultBuildSettings.SwitchUnhandledEnumeratorWarningLevel = WarningLevel.Warning;
+
+			ReadOnlyCppCompileWarnings wrapper = new(defaultBuildSettings);
+
+			CppCompileWarnings cloned = wrapper.CloneAsWriteable();
+
+			{
+				Assert.AreEqual(defaultBuildSettings.ShadowVariableWarningLevel, cloned.ShadowVariableWarningLevel);
+				Assert.AreEqual(defaultBuildSettings.UndefinedIdentifierWarningLevel, cloned.UndefinedIdentifierWarningLevel);
+				Assert.AreEqual(defaultBuildSettings.UnsafeTypeCastWarningLevel, cloned.UnsafeTypeCastWarningLevel);
+				Assert.AreEqual(defaultBuildSettings.SwitchUnhandledEnumeratorWarningLevel, cloned.SwitchUnhandledEnumeratorWarningLevel);
+			}
+
+			// Modify and assert != to each other - the same underyling reference has not been returned
+			{
+				cloned.ShadowVariableWarningLevel = WarningLevel.Warning;
+				cloned.UndefinedIdentifierWarningLevel = WarningLevel.Error;
+				cloned.UnsafeTypeCastWarningLevel = WarningLevel.Warning;
+				cloned.SwitchUnhandledEnumeratorWarningLevel = WarningLevel.Error;
+
+				Assert.AreNotEqual(defaultBuildSettings.ShadowVariableWarningLevel, cloned.ShadowVariableWarningLevel);
+				Assert.AreNotEqual(defaultBuildSettings.UndefinedIdentifierWarningLevel, cloned.UndefinedIdentifierWarningLevel);
+				Assert.AreNotEqual(defaultBuildSettings.UnsafeTypeCastWarningLevel, cloned.UnsafeTypeCastWarningLevel);
+				Assert.AreNotEqual(defaultBuildSettings.SwitchUnhandledEnumeratorWarningLevel, cloned.SwitchUnhandledEnumeratorWarningLevel);
+			}
+		}
+
+		[TestMethod]
+		public void TestContextBuildSettingsOverride()
+		{
+			TestBuildSettingsProvider buildContextProvider = new();
+			buildContextProvider._testBuildSettings = BuildSettingsVersion.V1;
+
+			CppCompileWarnings newWarnings = new(buildContextProvider, null);
+
+			// Shadow Variable build setting defaults
+			{
+				Assert.AreEqual(WarningLevel.Default, newWarnings.ShadowVariableWarningLevel);
+
+				// Boundary analysis
+				buildContextProvider._testBuildSettings = BuildSettingsVersion.V2;
+				CppCompileWarnings.ApplyTargetDefaults(newWarnings, true);
+				Assert.AreEqual(WarningLevel.Error, newWarnings.ShadowVariableWarningLevel);
+
+				buildContextProvider._testBuildSettings = BuildSettingsVersion.V3;
+
+				CppCompileWarnings.ApplyTargetDefaults(newWarnings, true);
+				Assert.AreEqual(WarningLevel.Error, newWarnings.ShadowVariableWarningLevel);
+			}
+		}
+
+		[TestMethod]
+		public void TestParentSettings()
+		{
+			TestBuildSettingsProvider buildContextProvider = new();
+			buildContextProvider._testBuildSettings = BuildSettingsVersion.V1;
+			CppCompileWarnings parentWarnings = new(buildContextProvider, null);
+			CppCompileWarnings newWarnings = new(buildContextProvider, null, new ReadOnlyCppCompileWarnings(parentWarnings));
+
+			// Assert invariants
+			{
+				Assert.AreEqual(WarningLevel.Default, newWarnings.ShadowVariableWarningLevel);
+				Assert.AreEqual(WarningLevel.Default, newWarnings.UnsafeTypeCastWarningLevel);
+				Assert.AreEqual(WarningLevel.Default, newWarnings.UndefinedIdentifierWarningLevel);
+			}
+
+			// Assert new parent values
+			{
+				parentWarnings.ShadowVariableWarningLevel = WarningLevel.Error;
+				parentWarnings.UndefinedIdentifierWarningLevel = WarningLevel.Error;
+				parentWarnings.UnsafeTypeCastWarningLevel = WarningLevel.Error;
+
+				Assert.AreEqual(WarningLevel.Error, newWarnings.ShadowVariableWarningLevel);
+				Assert.AreEqual(WarningLevel.Error, newWarnings.UnsafeTypeCastWarningLevel);
+				Assert.AreEqual(WarningLevel.Error, newWarnings.UndefinedIdentifierWarningLevel);
+			}
+		}
+
+		[TestMethod]
+		public void TestReparentOperations()
+		{
+			TestBuildSettingsProvider buildContextProvider = new();
+			buildContextProvider._testBuildSettings = BuildSettingsVersion.V1;
+			CppCompileWarnings parentWarnings = new(buildContextProvider, null);
+			CppCompileWarnings newWarnings = new(buildContextProvider, null, new ReadOnlyCppCompileWarnings(parentWarnings));
+
+			// Assert invariant parent values
+			{
+				parentWarnings.ShadowVariableWarningLevel = WarningLevel.Error;
+				parentWarnings.UndefinedIdentifierWarningLevel = WarningLevel.Error;
+				parentWarnings.UnsafeTypeCastWarningLevel = WarningLevel.Error;
+
+				Assert.AreEqual(WarningLevel.Error, newWarnings.ShadowVariableWarningLevel);
+				Assert.AreEqual(WarningLevel.Error, newWarnings.UnsafeTypeCastWarningLevel);
+				Assert.AreEqual(WarningLevel.Error, newWarnings.UndefinedIdentifierWarningLevel);
+			}
+
+			// Assert reparenting
+			CppCompileWarnings parentWarnings2 = new(buildContextProvider, null);
+			{
+				newWarnings.AddParent(parentWarnings2);
+				parentWarnings2.ShadowVariableWarningLevel = WarningLevel.Warning;
+				parentWarnings2.UndefinedIdentifierWarningLevel = WarningLevel.Warning;
+				parentWarnings2.UnsafeTypeCastWarningLevel = WarningLevel.Warning;
+
+				Assert.AreEqual(WarningLevel.Warning, newWarnings.ShadowVariableWarningLevel);
+				Assert.AreEqual(WarningLevel.Warning, newWarnings.UnsafeTypeCastWarningLevel);
+				Assert.AreEqual(WarningLevel.Warning, newWarnings.UndefinedIdentifierWarningLevel);
+			}
+
+			// Assert transitive to eldest ancestor
+			{
+				parentWarnings2.ShadowVariableWarningLevel = WarningLevel.Default;
+				Assert.AreEqual(parentWarnings.ShadowVariableWarningLevel, newWarnings.ShadowVariableWarningLevel);
+			}
+
+			// Assert removal of parentWarnings2 - UnsafeTypeCastWarningLevel & UndefinedIdentifierWarningLevel will now be the reoslved items
+			{
+				newWarnings.RemoveParent();
+				Assert.AreEqual(parentWarnings.UnsafeTypeCastWarningLevel, newWarnings.UnsafeTypeCastWarningLevel);
+				Assert.AreEqual(parentWarnings.UndefinedIdentifierWarningLevel, newWarnings.UndefinedIdentifierWarningLevel);
+
+			}
+
+			// Removal of only parent; should only be default values
+			{
+				newWarnings.RemoveParent();
+				Assert.AreEqual(WarningLevel.Default, newWarnings.ShadowVariableWarningLevel);
+				Assert.AreEqual(WarningLevel.Default, newWarnings.UnsafeTypeCastWarningLevel);
+				Assert.AreEqual(WarningLevel.Default, newWarnings.UndefinedIdentifierWarningLevel);
+			}
+		}
+
+		[TestMethod]
+		public void TestChildOverrides()
+		{
+			TestBuildSettingsProvider buildContextProvider = new();
+			buildContextProvider._testBuildSettings = BuildSettingsVersion.V1;
+			CppCompileWarnings parentWarnings = new(buildContextProvider, null);
+			CppCompileWarnings newWarnings = new(buildContextProvider, null, new ReadOnlyCppCompileWarnings(parentWarnings));
+
+			parentWarnings.ShadowVariableWarningLevel = WarningLevel.Error;
+			parentWarnings.UndefinedIdentifierWarningLevel = WarningLevel.Error;
+			parentWarnings.UnsafeTypeCastWarningLevel = WarningLevel.Error;
+
+			// Assert invariants
+			{
+				Assert.AreEqual(WarningLevel.Error, newWarnings.ShadowVariableWarningLevel);
+				Assert.AreEqual(WarningLevel.Error, newWarnings.UnsafeTypeCastWarningLevel);
+				Assert.AreEqual(WarningLevel.Error, newWarnings.UndefinedIdentifierWarningLevel);
+			}
+
+			// Child overrides;
+			newWarnings.ShadowVariableWarningLevel = WarningLevel.Warning;
+			newWarnings.UndefinedIdentifierWarningLevel = WarningLevel.Warning;
+			newWarnings.UnsafeTypeCastWarningLevel = WarningLevel.Warning;
+
+			// Assert new parent values
+			{
+				Assert.AreEqual(WarningLevel.Warning, newWarnings.ShadowVariableWarningLevel);
+				Assert.AreEqual(WarningLevel.Warning, newWarnings.UnsafeTypeCastWarningLevel);
+				Assert.AreEqual(WarningLevel.Warning, newWarnings.UndefinedIdentifierWarningLevel);
+			}
+		}
+
+		[TestMethod]
+		public void TestGenerateWarningCommandLineAgsOrder()
+		{
+			DirectoryReference testStubsDirectory = TestDirs.Stubs;
+			// Generate ephemeral TargetRules & uproject files
+			using (SafeTestDirectory testRootDirectory = TestDirs.Create())
+			using (SafeTestDirectory testSourceDirectory = SafeTestDirectory.CreateTestDirectory(Path.Combine(testRootDirectory.TemporaryDirectory, "Source")))
+			using (SafeTestFile testEditorTargetTestFile = new(File.ReadAllText(Path.Combine(testStubsDirectory.FullName, "TargetStubs", "TestStubEditor.Target.ubttest")), "TestStubEditor.Target.cs", testSourceDirectory.TemporaryDirectory))
+			using (SafeTestFile testTargetTestFile = new(File.ReadAllText(Path.Combine(testStubsDirectory.FullName, "TargetStubs", "TestStub.Target.ubttest")), "TestStub.Target.cs", testSourceDirectory.TemporaryDirectory))
+			using (SafeTestFile testUprojectTestFile = new(File.ReadAllText(Path.Combine(testStubsDirectory.FullName, "UprojectStubs", "TestStub.uproject.ubttest")), "TestStub.uproject", testRootDirectory.TemporaryDirectory))
+			{
+				CppCompileEnvironment newCompileEnvironment = new(UnrealTargetPlatform.Win64, CppConfiguration.Development, new UnrealArchitectures(UnrealArch.X64), null!);
+
+				FileReference testUproject = new(testUprojectTestFile.TemporaryFile);
+				IEnumerable<string> args = null!;
+
+				RulesAssembly generatedTestGameRulesAssembly = RulesCompiler.CreateTargetRulesAssembly(testUproject, "TestStub", false, false, false, null, false, new TestLogger());
+				TargetRules testTargetRulesNonWin64 = generatedTestGameRulesAssembly.CreateTargetRules("TestStub", UnrealTargetPlatform.Android, UnrealTargetConfiguration.Development, null, testUproject, null, new TestLogger());
+				CppCompileWarnings specializedWarnings = new(testTargetRulesNonWin64, null);
+
+				// Verify the ordering
+				{
+					VersionNumber newVersionNumber = new(18, 9, 9);
+					specializedWarnings.DeprecatedCopyWarningLevel = WarningLevel.Off;
+					specializedWarnings.DeprecatedCopyWithUserProvidedCopyWarningLevel = WarningLevel.Error;
+					args = specializedWarnings.GenerateWarningCommandLineArgs(newCompileEnvironment, typeof(ClangToolChain), newVersionNumber);
+
+					Assert.IsTrue(args.Count() == 2);
+					Assert.IsTrue(args.First() == "-Wno-deprecated-copy");
+					Assert.IsTrue(args.Last() == "-Wdeprecated-copy-with-user-provided-copy");
+				}
+			}
+		}
+
+		[TestMethod]
+		public void TestGenerateWarningCommandLineAgs()
+		{
+			DirectoryReference testStubsDirectory = TestDirs.Stubs;
+
+			// Generate ephemeral TargetRules & uproject files
+			using (SafeTestDirectory testRootDirectory = TestDirs.Create())
+			using (SafeTestDirectory testSourceDirectory = SafeTestDirectory.CreateTestDirectory(Path.Combine(testRootDirectory.TemporaryDirectory, "Source")))
+			using (SafeTestFile testEditorTargetTestFile = new(File.ReadAllText(Path.Combine(testStubsDirectory.FullName, "TargetStubs", "TestStubEditor.Target.ubttest")), "TestStubEditor.Target.cs", testSourceDirectory.TemporaryDirectory))
+			using (SafeTestFile testTargetTestFile = new(File.ReadAllText(Path.Combine(testStubsDirectory.FullName, "TargetStubs", "TestStub.Target.ubttest")), "TestStub.Target.cs", testSourceDirectory.TemporaryDirectory))
+			using (SafeTestFile testUprojectTestFile = new(File.ReadAllText(Path.Combine(testStubsDirectory.FullName, "UprojectStubs", "TestStub.uproject.ubttest")), "TestStub.uproject", testRootDirectory.TemporaryDirectory))
+			{
+				CppCompileEnvironment newCompileEnvironment = new(UnrealTargetPlatform.Win64, CppConfiguration.Development, new UnrealArchitectures(UnrealArch.X64), null!);
+
+				FileReference testUproject = new(testUprojectTestFile.TemporaryFile);
+
+				// Generate rules assembly & subsequent targets
+				RulesAssembly generatedTestRulesAssembly = RulesCompiler.CreateTargetRulesAssembly(testUproject, "TestStubEditor", false, false, false, null, false, new TestLogger());
+				TargetRules testTargetRules = generatedTestRulesAssembly.CreateTargetRules("TestStubEditor", UnrealTargetPlatform.Win64, UnrealTargetConfiguration.Development, null, testUproject, null, new TestLogger());
+				testTargetRules.WindowsPlatform.Compiler = WindowsCompiler.VisualStudio2022;
+				CppCompileWarnings baseWarnings = new(testTargetRules, null);
+
+				IEnumerable<string> args = null!;
+
+				// No args, base check
+				{
+					args = baseWarnings.GenerateWarningCommandLineArgs(newCompileEnvironment, typeof(VCToolChain));
+					Assert.IsFalse(args.Any());
+				}
+
+				// Specialized VC toolchain check; UndefinedIdentifier
+				{
+					baseWarnings.UndefinedIdentifierWarningLevel = WarningLevel.Warning;
+					args = baseWarnings.GenerateWarningCommandLineArgs(newCompileEnvironment, typeof(VCToolChain));
+
+					Assert.IsTrue(args.Count() == 1);
+					Assert.IsTrue(args.Contains("/w44668"));
+
+					CppCompileEnvironment preprocessOnlyCompileEnvironment = new(newCompileEnvironment);
+					preprocessOnlyCompileEnvironment.bPreprocessOnly = true;
+					args = baseWarnings.GenerateWarningCommandLineArgs(preprocessOnlyCompileEnvironment, typeof(VCToolChain));
+
+					Assert.IsFalse(args.Any());
+
+					baseWarnings.UndefinedIdentifierWarningLevel = WarningLevel.Default;
+				}
+
+				{
+					// DeterministicWarningLevel
+					VersionNumber newVersionNumber = new(18, 9, 9);
+
+					CppCompileEnvironment deterministicCompileEnvironment = new(UnrealTargetPlatform.Win64, CppConfiguration.Debug, new UnrealArchitectures(UnrealArch.X64), null!);
+					deterministicCompileEnvironment.bDeterministic = false;
+
+					args = baseWarnings.GenerateWarningCommandLineArgs(deterministicCompileEnvironment, typeof(VCToolChain), newVersionNumber);
+					Assert.IsFalse(args.Any());
+
+					deterministicCompileEnvironment.bDeterministic = true;
+
+					args = baseWarnings.GenerateWarningCommandLineArgs(deterministicCompileEnvironment, typeof(VCToolChain), newVersionNumber);
+					Assert.IsFalse(args.Any());
+
+					baseWarnings.DeterministicWarningLevel = WarningLevel.Off;
+					args = baseWarnings.GenerateWarningCommandLineArgs(deterministicCompileEnvironment, typeof(VCToolChain), newVersionNumber);
+
+					Assert.AreEqual(1, args.Count());
+					Assert.IsTrue(args.Contains("/wd5048"));
+
+					baseWarnings.DeterministicWarningLevel = WarningLevel.Warning;
+					args = baseWarnings.GenerateWarningCommandLineArgs(deterministicCompileEnvironment, typeof(VCToolChain), newVersionNumber);
+
+					Assert.AreEqual(0, args.Count());
+					Assert.IsFalse(args.Any());
+
+					baseWarnings.DeterministicWarningLevel = WarningLevel.Error;
+					args = baseWarnings.GenerateWarningCommandLineArgs(deterministicCompileEnvironment, typeof(VCToolChain), newVersionNumber);
+
+					Assert.AreEqual(1, args.Count());
+					Assert.IsTrue(args.Contains("/we5048"));
+				}
+			}
+		}
+
+		[TestMethod]
+		public void TestSpecializedClangCommandLineArgs()
+		{
+			DirectoryReference testStubsDirectory = TestDirs.Stubs;
+
+			// Generate ephemeral TargetRules & uproject files
+			using (SafeTestDirectory testRootDirectory = TestDirs.Create())
+			using (SafeTestDirectory testSourceDirectory = SafeTestDirectory.CreateTestDirectory(Path.Combine(testRootDirectory.TemporaryDirectory, "Source")))
+			using (SafeTestFile testEditorTargetTestFile = new(File.ReadAllText(Path.Combine(testStubsDirectory.FullName, "TargetStubs", "TestStubEditor.Target.ubttest")), "TestStubEditor.Target.cs", testSourceDirectory.TemporaryDirectory))
+			using (SafeTestFile testTargetTestFile = new(File.ReadAllText(Path.Combine(testStubsDirectory.FullName, "TargetStubs", "TestStub.Target.ubttest")), "TestStub.Target.cs", testSourceDirectory.TemporaryDirectory))
+			using (SafeTestFile testUprojectTestFile = new(File.ReadAllText(Path.Combine(testStubsDirectory.FullName, "UprojectStubs", "TestStub.uproject.ubttest")), "TestStub.uproject", testRootDirectory.TemporaryDirectory))
+			{
+				CppCompileEnvironment newCompileEnvironment = new(UnrealTargetPlatform.Win64, CppConfiguration.Development, new UnrealArchitectures(UnrealArch.X64), null!);
+
+				FileReference testUproject = new(testUprojectTestFile.TemporaryFile);
+
+				// Generate rules assembly & subsequent targets
+				RulesAssembly generatedTestRulesAssembly = RulesCompiler.CreateTargetRulesAssembly(testUproject, "TestStubEditor", false, false, false, null, false, new TestLogger());
+				TargetRules testTargetRules = generatedTestRulesAssembly.CreateTargetRules("TestStubEditor", UnrealTargetPlatform.Win64, UnrealTargetConfiguration.Development, null, testUproject, null, new TestLogger());
+				testTargetRules.WindowsPlatform.Compiler = WindowsCompiler.VisualStudio2022;
+				CppCompileWarnings baseWarnings = new(testTargetRules, null);
+
+				IEnumerable<string> args = null!;
+				{
+					// DeterministcFlagSetFilter
+					{
+						testTargetRules.WindowsPlatform.Compiler = WindowsCompiler.Clang;
+						baseWarnings.DeterministicWarningLevel = WarningLevel.Off;
+						VersionNumber newVersionNumber = new(18, 9, 9);
+
+						CppCompileEnvironment deterministicCompileEnvironment = new(UnrealTargetPlatform.Win64, CppConfiguration.Debug, new UnrealArchitectures(UnrealArch.X64), null!);
+						deterministicCompileEnvironment.bDeterministic = false;
+
+						args = baseWarnings.GenerateWarningCommandLineArgs(deterministicCompileEnvironment, typeof(ClangToolChain), newVersionNumber);
+						Assert.IsFalse(args.Any());
+
+						deterministicCompileEnvironment.bDeterministic = true;
+
+						args = baseWarnings.GenerateWarningCommandLineArgs(deterministicCompileEnvironment, typeof(ClangToolChain), newVersionNumber);
+						Assert.IsTrue(args.Count() == 1);
+						Assert.IsTrue(args.Contains("-Wno-date-time"));
+
+						deterministicCompileEnvironment.bDeterministic = false;
+						baseWarnings.DeterministicWarningLevel = WarningLevel.Warning;
+						args = baseWarnings.GenerateWarningCommandLineArgs(deterministicCompileEnvironment, typeof(ClangToolChain), newVersionNumber);
+						Assert.IsFalse(args.Any());
+
+						deterministicCompileEnvironment.bDeterministic = true;
+						args = baseWarnings.GenerateWarningCommandLineArgs(deterministicCompileEnvironment, typeof(ClangToolChain), newVersionNumber);
+
+						Assert.IsTrue(args.Count() == 2);
+						Assert.IsTrue(args.Contains("-Wdate-time"));
+						Assert.IsTrue(args.Contains("-Wno-error=date-time"));
+
+						baseWarnings.DeterministicWarningLevel = WarningLevel.Error;
+						args = baseWarnings.GenerateWarningCommandLineArgs(deterministicCompileEnvironment, typeof(ClangToolChain), newVersionNumber);
+						Assert.IsTrue(args.Count() == 1);
+						Assert.IsTrue(args.Contains("-Wdate-time"));
+
+						baseWarnings.DeterministicWarningLevel = WarningLevel.Default;
+					}
+
+					// PGOOptimizedFilter
+					{
+						testTargetRules.WindowsPlatform.Compiler = WindowsCompiler.Clang;
+						baseWarnings.ProfileInstructWarningLevel = WarningLevel.Off;
+						VersionNumber newVersionNumber = new(18, 9, 9);
+
+						CppCompileEnvironment pgoOptimizedCompileEnvironment = new(UnrealTargetPlatform.Win64, CppConfiguration.Debug, new UnrealArchitectures(UnrealArch.X64), null!);
+						pgoOptimizedCompileEnvironment.bPGOOptimize = false;
+
+						args = baseWarnings.GenerateWarningCommandLineArgs(pgoOptimizedCompileEnvironment, typeof(ClangToolChain), newVersionNumber);
+						Assert.IsFalse(args.Any());
+
+						pgoOptimizedCompileEnvironment.bPGOOptimize = true;
+
+						args = baseWarnings.GenerateWarningCommandLineArgs(pgoOptimizedCompileEnvironment, typeof(ClangToolChain), newVersionNumber);
+
+						Assert.IsTrue(args.Count() == 2);
+						Assert.IsTrue(args.Contains("-Wno-profile-instr-out-of-date"));
+						Assert.IsTrue(args.Contains("-Wno-profile-instr-unprofiled"));
+
+						baseWarnings.ProfileInstructWarningLevel = WarningLevel.Off;
+					}
+
+					// AndroidNDKR28ToolChainVersionExclusion
+					{
+						testTargetRules.WindowsPlatform.Compiler = WindowsCompiler.Clang;
+						baseWarnings.CastFunctionTypeMismatchWarningLevel = WarningLevel.Off;
+
+						VersionNumber newVersionNumber = new(19, 0, 0);
+						CppCompileEnvironment androidCompileEnvironment = new(UnrealTargetPlatform.Android, CppConfiguration.Development, new UnrealArchitectures(UnrealArch.X64), null!);
+						args = baseWarnings.GenerateWarningCommandLineArgs(androidCompileEnvironment, typeof(ClangToolChain), newVersionNumber);
+
+						Assert.IsFalse(args.Any());
+
+						args = baseWarnings.GenerateWarningCommandLineArgs(newCompileEnvironment, typeof(ClangToolChain), newVersionNumber);
+
+						Assert.IsTrue(args.Count() == 1);
+						Assert.IsTrue(args.Contains("-Wno-cast-function-type-mismatch"));
+
+						newVersionNumber = new VersionNumber(19, 0, 1);
+						args = baseWarnings.GenerateWarningCommandLineArgs(androidCompileEnvironment, typeof(ClangToolChain), newVersionNumber);
+
+						Assert.IsTrue(args.Count() == 1);
+						Assert.IsTrue(args.Contains("-Wno-cast-function-type-mismatch"));
+
+						baseWarnings.CastFunctionTypeMismatchWarningLevel = WarningLevel.Default;
+					}
+
+					// AndroidNDKR26ToolChainVersionExclusion
+					{
+						testTargetRules.WindowsPlatform.Compiler = WindowsCompiler.Clang;
+						baseWarnings.InvalidUnevaluatedStringWarningLevel = WarningLevel.Off;
+						CppCompileEnvironment androidCompileEnvironment = new(UnrealTargetPlatform.Android, CppConfiguration.Development, new UnrealArchitectures(UnrealArch.X64), null!);
+
+						// Exclusion cases
+						{
+							VersionNumber newVersionNumber = new(17, 0, 2);
+							args = baseWarnings.GenerateWarningCommandLineArgs(androidCompileEnvironment, typeof(ClangToolChain), newVersionNumber);
+
+							Assert.IsTrue(args.Count() == 1 && args.Contains("-Wno-shadow"));
+
+							newVersionNumber = new VersionNumber(16);
+							args = baseWarnings.GenerateWarningCommandLineArgs(androidCompileEnvironment, typeof(ClangToolChain), newVersionNumber);
+
+							Assert.IsFalse(args.Any());
+
+							newVersionNumber = new VersionNumber(17, 1);
+							androidCompileEnvironment.CppStandard = CppStandardVersion.Latest;
+							args = baseWarnings.GenerateWarningCommandLineArgs(androidCompileEnvironment, typeof(ClangToolChain), newVersionNumber);
+
+							Assert.IsTrue(args.Count() == 1 && args.Contains("-Wno-shadow"));
+						}
+
+						// Inclusion cases
+						{
+							VersionNumber newVersionNumber = new(17, 1);
+							androidCompileEnvironment.CppStandard = CppStandardVersion.Cpp20;
+
+							args = baseWarnings.GenerateWarningCommandLineArgs(androidCompileEnvironment, typeof(ClangToolChain), newVersionNumber);
+
+							Assert.IsTrue(args.Contains("-Wno-invalid-unevaluated-string"));
+
+							newVersionNumber = new VersionNumber(17, 0, 2);
+
+							args = baseWarnings.GenerateWarningCommandLineArgs(newCompileEnvironment, typeof(ClangToolChain), newVersionNumber);
+
+							Assert.IsTrue(args.Contains("-Wno-invalid-unevaluated-string"));
+						}
+
+						baseWarnings.InvalidUnevaluatedStringWarningLevel = WarningLevel.Default;
+					}
+
+					// UnusedValueClangToolChainAttribute
+					{
+						testTargetRules.WindowsPlatform.Compiler = WindowsCompiler.Clang;
+						baseWarnings.UnusedValueWarningLevel = WarningLevel.Off;
+						VersionNumber newVersionNumber = new(18, 9, 9);
+
+						CppCompileEnvironment invalidCompileEnvironment = new(UnrealTargetPlatform.Win64, CppConfiguration.Debug, new UnrealArchitectures(UnrealArch.X64), null!);
+						CppCompileEnvironment validCompileEnvironment = new(UnrealTargetPlatform.Win64, CppConfiguration.Shipping, new UnrealArchitectures(UnrealArch.X64), null!);
+
+						args = baseWarnings.GenerateWarningCommandLineArgs(invalidCompileEnvironment, typeof(ClangToolChain), newVersionNumber, StaticAnalyzer.None);
+						Assert.IsFalse(args.Any());
+
+						args = baseWarnings.GenerateWarningCommandLineArgs(validCompileEnvironment, typeof(ClangToolChain), newVersionNumber, StaticAnalyzer.None);
+						Assert.IsTrue(args.Count() == 2 && args.Contains("-Wno-unused-value"));
+
+						args = baseWarnings.GenerateWarningCommandLineArgs(validCompileEnvironment, typeof(ClangToolChain), newVersionNumber, StaticAnalyzer.Clang);
+						Assert.IsTrue(args.Count() == 2 && args.Contains("-Wno-unused-value"));
+					}
+				}
+			}
+		}
+
+		[TestMethod]
+		public void TestConfigurationFilterCommandLineArgs()
+		{
+			DirectoryReference testStubsDirectory = TestDirs.Stubs;
+
+			// Generate ephemeral TargetRules & uproject files
+			using (SafeTestDirectory testRootDirectory = TestDirs.Create())
+			using (SafeTestDirectory testSourceDirectory = SafeTestDirectory.CreateTestDirectory(Path.Combine(testRootDirectory.TemporaryDirectory, "Source")))
+			using (SafeTestFile testEditorTargetTestFile = new(File.ReadAllText(Path.Combine(testStubsDirectory.FullName, "TargetStubs", "TestStubEditor.Target.ubttest")), "TestStubEditor.Target.cs", testSourceDirectory.TemporaryDirectory))
+			using (SafeTestFile testTargetTestFile = new(File.ReadAllText(Path.Combine(testStubsDirectory.FullName, "TargetStubs", "TestStub.Target.ubttest")), "TestStub.Target.cs", testSourceDirectory.TemporaryDirectory))
+			using (SafeTestFile testUprojectTestFile = new(File.ReadAllText(Path.Combine(testStubsDirectory.FullName, "UprojectStubs", "TestStub.uproject.ubttest")), "TestStub.uproject", testRootDirectory.TemporaryDirectory))
+			{
+				FileReference testUproject = new(testUprojectTestFile.TemporaryFile);
+				RulesAssembly generatedTestRulesAssembly = RulesCompiler.CreateTargetRulesAssembly(testUproject, "TestStubEditor", false, false, false, null, false, new TestLogger());
+				TargetRules testTargetRules = generatedTestRulesAssembly.CreateTargetRules("TestStubEditor", UnrealTargetPlatform.Win64, UnrealTargetConfiguration.Development, null, testUproject, null, new TestLogger());
+				testTargetRules.WindowsPlatform.Compiler = WindowsCompiler.VisualStudio2022;
+				CppCompileWarnings baseWarnings = new(testTargetRules, null);
+
+				CppCompileEnvironment newCompileEnvironment = new(UnrealTargetPlatform.Win64, CppConfiguration.Development, new UnrealArchitectures(UnrealArch.X64), null!);
+				IEnumerable<string> args = null!;
+
+				// ConfigurationFilterAttribute
+				{
+					testTargetRules.WindowsPlatform.Compiler = WindowsCompiler.Clang;
+					baseWarnings.UnusedValueWarningLevel = WarningLevel.Off;
+					VersionNumber newVersionNumber = new(18, 9, 9);
+
+					CppCompileEnvironment invalidCompileEnvironment = new(UnrealTargetPlatform.Win64, CppConfiguration.Debug, new UnrealArchitectures(UnrealArch.X64), null!);
+					CppCompileEnvironment validCompileEnvironment = new(UnrealTargetPlatform.Win64, CppConfiguration.Shipping, new UnrealArchitectures(UnrealArch.X64), null!);
+
+					args = baseWarnings.GenerateWarningCommandLineArgs(invalidCompileEnvironment, typeof(ClangToolChain), newVersionNumber);
+					Assert.IsFalse(args.Any());
+
+					args = baseWarnings.GenerateWarningCommandLineArgs(validCompileEnvironment, typeof(ClangToolChain), newVersionNumber);
+					Assert.IsTrue(args.Count() == 2);
+					Assert.IsTrue(args.Contains("-Wno-unused-value"));
+
+					baseWarnings.UnusedValueWarningLevel = WarningLevel.Default;
+				}
+			}
+		}
+
+		[TestMethod]
+		public void TestCppStandardConstraintCommandLineArgs()
+		{
+			DirectoryReference testStubsDirectory = TestDirs.Stubs;
+
+			// Generate ephemeral TargetRules & uproject files
+			using (SafeTestDirectory testRootDirectory = TestDirs.Create())
+			using (SafeTestDirectory testSourceDirectory = SafeTestDirectory.CreateTestDirectory(Path.Combine(testRootDirectory.TemporaryDirectory, "Source")))
+			using (SafeTestFile testEditorTargetTestFile = new(File.ReadAllText(Path.Combine(testStubsDirectory.FullName, "TargetStubs", "TestStubEditor.Target.ubttest")), "TestStubEditor.Target.cs", testSourceDirectory.TemporaryDirectory))
+			using (SafeTestFile testTargetTestFile = new(File.ReadAllText(Path.Combine(testStubsDirectory.FullName, "TargetStubs", "TestStub.Target.ubttest")), "TestStub.Target.cs", testSourceDirectory.TemporaryDirectory))
+			using (SafeTestFile testUprojectTestFile = new(File.ReadAllText(Path.Combine(testStubsDirectory.FullName, "UprojectStubs", "TestStub.uproject.ubttest")), "TestStub.uproject", testRootDirectory.TemporaryDirectory))
+			{
+				FileReference testUproject = new(testUprojectTestFile.TemporaryFile);
+				RulesAssembly generatedTestRulesAssembly = RulesCompiler.CreateTargetRulesAssembly(testUproject, "TestStubEditor", false, false, false, null, false, new TestLogger());
+				TargetRules testTargetRules = generatedTestRulesAssembly.CreateTargetRules("TestStubEditor", UnrealTargetPlatform.Win64, UnrealTargetConfiguration.Development, null, testUproject, null, new TestLogger());
+				testTargetRules.WindowsPlatform.Compiler = WindowsCompiler.VisualStudio2022;
+				CppCompileWarnings baseWarnings = new(testTargetRules, null);
+
+				CppCompileEnvironment newCompileEnvironment = new(UnrealTargetPlatform.Win64, CppConfiguration.Development, new UnrealArchitectures(UnrealArch.X64), null!);
+				IEnumerable<string> args = null!;
+
+				// CppStandardConstrainedFilterAttribute
+				{
+					testTargetRules.WindowsPlatform.Compiler = WindowsCompiler.Clang;
+					baseWarnings.AmbiguousReversedOperatorWarningLevel = WarningLevel.Off;
+					VersionNumber newVersionNumber = new(18, 9, 9);
+
+					CppCompileEnvironment cppStandardCompileEnvironment = new(UnrealTargetPlatform.Win64, CppConfiguration.Debug, new UnrealArchitectures(UnrealArch.X64), null!);
+#pragma warning disable CS0618 // Type or member is obsolete
+					cppStandardCompileEnvironment.CppStandard = CppStandardVersion.Cpp17;
+#pragma warning restore CS0618 // Type or member is obsolete
+
+					args = baseWarnings.GenerateWarningCommandLineArgs(cppStandardCompileEnvironment, typeof(ClangToolChain), newVersionNumber);
+					Assert.IsFalse(args.Any());
+
+					cppStandardCompileEnvironment.CppStandard = CppStandardVersion.Cpp20;
+					args = baseWarnings.GenerateWarningCommandLineArgs(cppStandardCompileEnvironment, typeof(ClangToolChain), newVersionNumber);
+					Assert.IsTrue(args.Count() == 1);
+					Assert.IsTrue(args.Contains("-Wno-ambiguous-reversed-operator"));
+
+					cppStandardCompileEnvironment.CppStandard = CppStandardVersion.Latest;
+					args = baseWarnings.GenerateWarningCommandLineArgs(cppStandardCompileEnvironment, typeof(ClangToolChain), newVersionNumber);
+					Assert.IsTrue(args.Count() == 1);
+					Assert.IsTrue(args.Contains("-Wno-ambiguous-reversed-operator"));
+
+					baseWarnings.AmbiguousReversedOperatorWarningLevel = WarningLevel.Default;
+				}
+			}
+		}
+
+		[TestMethod]
+		public void TestVersionConstraintCommandLineArgs()
+		{
+			DirectoryReference testStubsDirectory = TestDirs.Stubs;
+
+			// Generate ephemeral TargetRules & uproject files
+			using (SafeTestDirectory testRootDirectory = TestDirs.Create())
+			using (SafeTestDirectory testSourceDirectory = SafeTestDirectory.CreateTestDirectory(Path.Combine(testRootDirectory.TemporaryDirectory, "Source")))
+			using (SafeTestFile testEditorTargetTestFile = new(File.ReadAllText(Path.Combine(testStubsDirectory.FullName, "TargetStubs", "TestStubEditor.Target.ubttest")), "TestStubEditor.Target.cs", testSourceDirectory.TemporaryDirectory))
+			using (SafeTestFile testTargetTestFile = new(File.ReadAllText(Path.Combine(testStubsDirectory.FullName, "TargetStubs", "TestStub.Target.ubttest")), "TestStub.Target.cs", testSourceDirectory.TemporaryDirectory))
+			using (SafeTestFile testUprojectTestFile = new(File.ReadAllText(Path.Combine(testStubsDirectory.FullName, "UprojectStubs", "TestStub.uproject.ubttest")), "TestStub.uproject", testRootDirectory.TemporaryDirectory))
+			{
+				FileReference testUproject = new(testUprojectTestFile.TemporaryFile);
+				RulesAssembly generatedTestRulesAssembly = RulesCompiler.CreateTargetRulesAssembly(testUproject, "TestStubEditor", false, false, false, null, false, new TestLogger());
+				TargetRules testTargetRules = generatedTestRulesAssembly.CreateTargetRules("TestStubEditor", UnrealTargetPlatform.Win64, UnrealTargetConfiguration.Development, null, testUproject, null, new TestLogger());
+				testTargetRules.WindowsPlatform.Compiler = WindowsCompiler.VisualStudio2022;
+				CppCompileWarnings baseWarnings = new(testTargetRules, null);
+
+				CppCompileEnvironment newCompileEnvironment = new(UnrealTargetPlatform.Win64, CppConfiguration.Development, new UnrealArchitectures(UnrealArch.X64), null!);
+				IEnumerable<string> args = null!;
+
+				// Specialized VC toolchain check;  UnsafeTypeCast
+				{
+					RulesAssembly generatedTestGameRulesAssembly = RulesCompiler.CreateTargetRulesAssembly(testUproject, "TestStub", false, false, false, null, false, new TestLogger());
+					TargetRules testTargetRulesNonWin64 = generatedTestGameRulesAssembly.CreateTargetRules("TestStub", UnrealTargetPlatform.Android, UnrealTargetConfiguration.Development, null, testUproject, null, new TestLogger());
+					CppCompileWarnings specializedWarnings = new(testTargetRulesNonWin64, null);
+
+					specializedWarnings.UnsafeTypeCastWarningLevel = WarningLevel.Warning;
+					args = specializedWarnings.GenerateWarningCommandLineArgs(newCompileEnvironment, typeof(VCToolChain));
+
+					Assert.IsTrue(args.Count() == 2);
+					Assert.IsTrue(args.Contains("/wd4244"));
+					Assert.IsTrue(args.Contains("/wd4838"));
+
+					baseWarnings.UnsafeTypeCastWarningLevel = WarningLevel.Warning;
+					args = baseWarnings.GenerateWarningCommandLineArgs(newCompileEnvironment, typeof(VCToolChain));
+
+					Assert.IsTrue(args.Count() == 2);
+					Assert.IsTrue(args.Contains("/w44244"));
+					Assert.IsTrue(args.Contains("/w44838"));
+
+					baseWarnings.UnsafeTypeCastWarningLevel = WarningLevel.Default;
+				}
+
+				// Clang toolchain version check - no version; should disable Shadow
+				{
+					testTargetRules.WindowsPlatform.Compiler = WindowsCompiler.Clang;
+					baseWarnings.ShadowVariableWarningLevel = WarningLevel.Error;
+
+					args = baseWarnings.GenerateWarningCommandLineArgs(newCompileEnvironment, typeof(ClangToolChain));
+					Assert.IsTrue(args.Count() == 1);
+					Assert.IsTrue(args.Contains("-Wno-shadow"));
+
+					baseWarnings.ShadowVariableWarningLevel = WarningLevel.Default;
+				}
+
+				// Clang toolchain version check - valid versions;
+				{
+					testTargetRules.WindowsPlatform.Compiler = WindowsCompiler.Clang;
+					baseWarnings.ShadowVariableWarningLevel = WarningLevel.Error;
+
+					VersionNumber newVersionNumber = new(16, 9, 5);
+
+					args = baseWarnings.GenerateWarningCommandLineArgs(newCompileEnvironment, typeof(ClangToolChain), newVersionNumber);
+					Assert.IsTrue(args.Count() == 1);
+					Assert.IsTrue(args.Contains("-Wshadow"));
+
+					newVersionNumber = new VersionNumber(17);
+					args = baseWarnings.GenerateWarningCommandLineArgs(newCompileEnvironment, typeof(ClangToolChain), newVersionNumber);
+					Assert.IsTrue(args.Count() == 1);
+					Assert.IsTrue(args.Contains("-Wno-shadow"));
+
+					newVersionNumber = new VersionNumber(17, 1);
+					args = baseWarnings.GenerateWarningCommandLineArgs(newCompileEnvironment, typeof(ClangToolChain), newVersionNumber);
+					Assert.IsTrue(args.Count() == 1);
+					Assert.IsTrue(args.Contains("-Wno-shadow"));
+
+					newVersionNumber = new VersionNumber(18, 1, 2);
+					args = baseWarnings.GenerateWarningCommandLineArgs(newCompileEnvironment, typeof(ClangToolChain), newVersionNumber);
+					Assert.IsTrue(args.Count() == 1);
+					Assert.IsTrue(args.Contains("-Wno-shadow"));
+
+					newVersionNumber = new VersionNumber(18, 1, 3);
+					args = baseWarnings.GenerateWarningCommandLineArgs(newCompileEnvironment, typeof(ClangToolChain), newVersionNumber);
+					Assert.IsTrue(args.Count() == 1);
+					Assert.IsTrue(args.Contains("-Wshadow"));
+
+					baseWarnings.ShadowVariableWarningLevel = WarningLevel.Default;
+				}
+
+				// ToolChainVersionConstrainedFilterAttribute
+				{
+					testTargetRules.WindowsPlatform.Compiler = WindowsCompiler.Clang;
+					baseWarnings.Shorten64To32WarningLevel = WarningLevel.Off;
+
+					VersionNumber newVersionNumber = new(18, 9, 9);
+					args = baseWarnings.GenerateWarningCommandLineArgs(newCompileEnvironment, typeof(ClangToolChain), newVersionNumber);
+					Assert.IsFalse(args.Any());
+
+					//
+					newVersionNumber = new VersionNumber(19, 0, 0);
+					args = baseWarnings.GenerateWarningCommandLineArgs(newCompileEnvironment, typeof(ClangToolChain), newVersionNumber);
+					Assert.IsTrue(args.Count() == 1);
+					Assert.IsTrue(args.Contains("-Wno-shorten-64-to-32"));
+					baseWarnings.Shorten64To32WarningLevel = WarningLevel.Default;
+				}
+			}
+		}
+
+		[TestMethod]
+		public void TestMSVCAndVCClangGenerateWarningCommandLineArgs()
+		{
+			DirectoryReference testStubsDirectory = TestDirs.Stubs;
+
+			// Generate ephemeral TargetRules & uproject files
+			using (SafeTestDirectory testRootDirectory = TestDirs.Create())
+			using (SafeTestDirectory testSourceDirectory = SafeTestDirectory.CreateTestDirectory(Path.Combine(testRootDirectory.TemporaryDirectory, "Source")))
+			using (SafeTestFile testEditorTargetTestFile = new(File.ReadAllText(Path.Combine(testStubsDirectory.FullName, "TargetStubs", "TestStubEditor.Target.ubttest")), "TestStubEditor.Target.cs", testSourceDirectory.TemporaryDirectory))
+			using (SafeTestFile testTargetTestFile = new(File.ReadAllText(Path.Combine(testStubsDirectory.FullName, "TargetStubs", "TestStub.Target.ubttest")), "TestStub.Target.cs", testSourceDirectory.TemporaryDirectory))
+			using (SafeTestFile testUprojectTestFile = new(File.ReadAllText(Path.Combine(testStubsDirectory.FullName, "UprojectStubs", "TestStub.uproject.ubttest")), "TestStub.uproject", testRootDirectory.TemporaryDirectory))
+			{
+				FileReference testUproject = new(testUprojectTestFile.TemporaryFile);
+				RulesAssembly generatedTestRulesAssembly = RulesCompiler.CreateTargetRulesAssembly(testUproject, "TestStubEditor", false, false, false, null, false, new TestLogger());
+				TargetRules testTargetRules = generatedTestRulesAssembly.CreateTargetRules("TestStubEditor", UnrealTargetPlatform.Win64, UnrealTargetConfiguration.Development, null, testUproject, null, new TestLogger());
+				testTargetRules.WindowsPlatform.Compiler = WindowsCompiler.VisualStudio2022;
+				CppCompileWarnings baseWarnings = new(testTargetRules, null);
+
+				CppCompileEnvironment newCompileEnvironment = new(UnrealTargetPlatform.Win64, CppConfiguration.Development, new UnrealArchitectures(UnrealArch.X64), null!);
+				IEnumerable<string> args = null!;
+
+				// VC toolchain check;
+				{
+					baseWarnings.UnsafeTypeCastWarningLevel = WarningLevel.Warning;
+					args = baseWarnings.GenerateWarningCommandLineArgs(newCompileEnvironment, typeof(VCToolChain));
+					Assert.IsTrue(args.Count() == 2);
+					Assert.IsTrue(args.Contains("/w44244"));
+					Assert.IsTrue(args.Contains("/w44838"));
+
+					baseWarnings.UnsafeTypeCastWarningLevel = WarningLevel.Default;
+				}
+
+				// MSVC compiler check;
+				{
+					baseWarnings.ShadowVariableWarningLevel = WarningLevel.Error;
+					args = baseWarnings.GenerateWarningCommandLineArgs(newCompileEnvironment, typeof(VCToolChain));
+
+					Assert.IsTrue(args.Count() == 3);
+					Assert.IsTrue(args.Contains("/we4456"));
+					Assert.IsTrue(args.Contains("/we4458"));
+					Assert.IsTrue(args.Contains("/we4459"));
+
+					WindowsCompiler previousCompiler = testTargetRules.WindowsPlatform.Compiler;
+					testTargetRules.WindowsPlatform.Compiler = WindowsCompiler.Intel;
+
+					args = baseWarnings.GenerateWarningCommandLineArgs(newCompileEnvironment, typeof(VCToolChain));
+
+					Assert.IsFalse(args.Any());
+
+					testTargetRules.WindowsPlatform.Compiler = previousCompiler;
+					baseWarnings.ShadowVariableWarningLevel = WarningLevel.Default;
+				}
+
+				// VCClangCompilerFilter
+				{
+					testTargetRules.WindowsPlatform.Compiler = WindowsCompiler.Clang;
+					baseWarnings.MicrosoftGroupWarningLevel = WarningLevel.Off;
+
+					args = baseWarnings.GenerateWarningCommandLineArgs(newCompileEnvironment, typeof(VCToolChain));
+					Assert.IsTrue(args.Any());
+					Assert.IsTrue(args.Contains("-Wno-microsoft"));
+
+					// make sure exclusively ONLY the clang stuff is in:
+					testTargetRules.WindowsPlatform.Compiler = WindowsCompiler.VisualStudio2022;
+					args = baseWarnings.GenerateWarningCommandLineArgs(newCompileEnvironment, typeof(VCToolChain));
+					Assert.IsFalse(args.Any());
+				}
+			}
+		}
+	}
+}
\ No newline at end of file
diff --git a/Engine/Source/Programs/UnrealBuildTool.Tests/JsonExtensionsTests.cs b/Engine/Source/Programs/UnrealBuildTool.Tests/JsonExtensionsTests.cs
new file mode 100644
--- /dev/null
+++ b/Engine/Source/Programs/UnrealBuildTool.Tests/JsonExtensionsTests.cs
@@ -0,0 +1,93 @@
+// Copyright Epic Games, Inc. All Rights Reserved.
+
+using System;
+using EpicGames.Core;
+using JsonExtensions;
+using Microsoft.VisualStudio.TestTools.UnitTesting;
+using UnrealBuildBase;
+
+namespace UnrealBuildTool.Tests
+{
+	[TestClass]
+	public class JsonExtensionsTests
+	{
+		public enum TestEnum
+		{
+			A,
+			B,
+			C,
+		}
+
+		[TestMethod]
+		public void TestGetEnumArrayFieldWithDeprecatedFallbackReturnsNullWhenBothAbsent()
+		{
+			JsonObject json = new();
+			TestEnum[]? values = json.GetEnumArrayFieldWithDeprecatedFallback<TestEnum>("New", "Old");
+			Assert.IsNull(values);
+		}
+
+		[TestMethod]
+		public void TestGetEnumArrayFieldWithDeprecatedFallbackThrowsWhenBothPresent()
+		{
+			JsonObject json = new();
+			json.AddOrSetFieldValue("New", [TestEnum.A, TestEnum.B]);
+			json.AddOrSetFieldValue("Old", [TestEnum.A, TestEnum.B]);
+			BuildLogEventException ex = Assert.ThrowsExactly<BuildLogEventException>(
+				() => json.GetEnumArrayFieldWithDeprecatedFallback<TestEnum>("New", "Old"));
+			StringAssert.Contains(ex.Message, "New", StringComparison.Ordinal);
+			StringAssert.Contains(ex.Message, "Old", StringComparison.Ordinal);
+		}
+
+		[TestMethod]
+		public void TestGetEnumArrayFieldWithDeprecatedFallbackPicksUpNewValues()
+		{
+			JsonObject json = new();
+			json.AddOrSetFieldValue("New", [TestEnum.A, TestEnum.B]);
+			TestEnum[]? values = json.GetEnumArrayFieldWithDeprecatedFallback<TestEnum>("New", "Old");
+			CollectionAssert.AreEquivalent(values, new[] { TestEnum.A, TestEnum.B });
+		}
+
+		[TestMethod]
+		public void TestGetEnumArrayFieldWithDeprecatedFallbackPicksUpFallbackValues()
+		{
+			JsonObject json = new();
+			json.AddOrSetFieldValue("Old", [TestEnum.A, TestEnum.B]);
+			TestEnum[]? values = json.GetEnumArrayFieldWithDeprecatedFallback<TestEnum>("New", "Old");
+			CollectionAssert.AreEquivalent(values, new[] { TestEnum.A, TestEnum.B });
+		}
+
+		[TestMethod]
+		public void TestGetEnumArrayFieldWithDeprecatedFallbackFailsOnInvalidNewValue()
+		{
+			JsonObject json = new();
+			json.AddOrSetFieldValue("New", ["A", "B", "C", "BadItem1", "BadItem2"]);
+			BuildLogEventException ex = Assert.ThrowsExactly<BuildLogEventException>(
+				() => json.GetEnumArrayFieldWithDeprecatedFallback<TestEnum>("New", "Old"));
+			StringAssert.Contains(ex.Message, "New", StringComparison.Ordinal);
+			StringAssert.Contains(ex.Message, "BadItem1", StringComparison.Ordinal);
+			StringAssert.Contains(ex.Message, "BadItem2", StringComparison.Ordinal);
+		}
+
+		[TestMethod]
+		public void TestGetEnumArrayFieldWithDeprecatedFallbackFailsOnInvalidFallbackValue()
+		{
+			JsonObject json = new();
+			json.AddOrSetFieldValue("Old", ["A", "B", "C", "BadItem1", "BadItem2"]);
+			BuildLogEventException ex = Assert.ThrowsExactly<BuildLogEventException>(
+				() => json.GetEnumArrayFieldWithDeprecatedFallback<TestEnum>("New", "Old"));
+			StringAssert.Contains(ex.Message, "Old", StringComparison.Ordinal);
+			StringAssert.Contains(ex.Message, "BadItem1", StringComparison.Ordinal);
+			StringAssert.Contains(ex.Message, "BadItem2", StringComparison.Ordinal);
+		}
+
+		[TestMethod]
+		public void TestGetEnumArrayFieldWithDeprecatedFallbackFailsOnInvalidDataType()
+		{
+			JsonObject json = new();
+			json.AddOrSetFieldValue("New", "A");
+			BuildLogEventException ex = Assert.ThrowsExactly<BuildLogEventException>(
+				() => json.GetEnumArrayFieldWithDeprecatedFallback<TestEnum>("New", "Old"));
+			StringAssert.Contains(ex.Message, "New", StringComparison.Ordinal);
+		}
+	}
+}
diff --git a/Engine/Source/Programs/UnrealBuildTool.Tests/PluginDescriptorTests.cs b/Engine/Source/Programs/UnrealBuildTool.Tests/PluginDescriptorTests.cs
new file mode 100644
--- /dev/null
+++ b/Engine/Source/Programs/UnrealBuildTool.Tests/PluginDescriptorTests.cs
@@ -0,0 +1,353 @@
+// Copyright Epic Games, Inc. All Rights Reserved.
+
+using System;
+using System.Collections.Generic;
+using System.Diagnostics;
+using System.IO;
+using System.Linq;
+using System.Text;
+using EpicGames.Core;
+using Microsoft.VisualStudio.TestTools.UnitTesting;
+
+namespace UnrealBuildTool.Tests
+{
+	[TestClass]
+	public class PluginDescriptorTests
+	{
+		private static DirectoryReference? s_tempDirectory;
+
+		private static DirectoryReference CreateTempDir()
+		{
+			string tempDir = Path.Join(Path.GetTempPath(), "epicgames-core-tests-" + Guid.NewGuid().ToString()[..8]);
+			Directory.CreateDirectory(tempDir);
+			return new DirectoryReference(tempDir);
+		}
+
+		[ClassInitialize]
+		public static void Setup(TestContext _)
+		{
+			PluginDescriptorTests.s_tempDirectory = CreateTempDir();
+		}
+
+		[ClassCleanup]
+		public static void TearDown()
+		{
+			if (s_tempDirectory != null && Directory.Exists(PluginDescriptorTests.s_tempDirectory.FullName))
+			{
+				Directory.Delete(PluginDescriptorTests.s_tempDirectory.FullName, true);
+			}
+		}
+
+		[TestMethod]
+		public void ReadPluginCustomFields()
+		{
+			string json = CreateUPluginWithCustomFields();
+			string filename = "CustomFields.uplugin";
+			FileReference fileToWrite = CreateTempUPluginFile(filename, json);
+			PluginDescriptor.FromFile(fileToWrite);
+		}
+
+		[TestMethod]
+		public void ReadUpluginMissingFileVersion()
+		{
+			string json = CreateUpluginWithMissingVersion();
+			string filename = "MissingFileVersionField.uplugin";
+			FileReference fileToWrite = CreateTempUPluginFile(filename, json);
+			Assert.ThrowsExactly<BuildException>(() => PluginDescriptor.FromFile(fileToWrite));
+		}
+
+		[TestMethod]
+		public void ReadUpluginInvalidVersion()
+		{
+			string json = CreateUpluginWithInvalidVersion();
+			string filename = "InvalidFileVersionField.uplugin";
+			FileReference fileToWrite = CreateTempUPluginFile(filename, json);
+			Assert.ThrowsExactly<BuildException>(() => PluginDescriptor.FromFile(fileToWrite));
+		}
+
+		[TestMethod]
+		public void Save2CustomFieldsMissingDefaults()
+		{
+			string json = CreateUPluginWithCustomFields();
+			string filename = "CustomFields-Save.uplugin";
+			FileReference fileToWrite = CreateTempUPluginFile(filename, json);
+			PluginDescriptor descriptor = PluginDescriptor.FromFile(fileToWrite);
+			string fileToSave = GetAbsolutePathToTempFile("CustomFields-SaveResults.uplugin");
+			descriptor.Save2(fileToSave);
+
+			string saveContents = File.ReadAllText(fileToSave);
+			JsonObject writtenObject = JsonObject.Parse(json);
+			JsonObject savedObject = JsonObject.Parse(saveContents);
+			IEnumerable<string> writtenObjectKeys = writtenObject.KeyNames;
+			// At this point, the written object should contain all the custom fields 
+			// The saved object should have all the custom fields as well as the missing default fields appended
+			// This means we can't do a simple equality check to see if the objects are the same 
+			// We check to make sure that all of the keys in the written object exist in the saved object as a short hand to check that all the custom fields have been saved and carried over to the saved object 
+			foreach (string key in writtenObjectKeys)
+			{
+				Assert.IsTrue(savedObject.ContainsField(key));
+			}
+		}
+
+		[TestMethod]
+		public void Save2DefaultFields()
+		{
+			string json = CreateUpluginWithDefaultFields();
+			string filename = "DefaultFields-Save.uplugin";
+			FileReference fileToWrite = CreateTempUPluginFile(filename, json);
+			PluginDescriptor descriptor = PluginDescriptor.FromFile(fileToWrite);
+			string fileToSave = GetAbsolutePathToTempFile("DefaultFields-SaveResults.uplugin");
+			descriptor.Save2(fileToSave);
+			string saveContents = File.ReadAllText(fileToSave);
+
+			// Nothing should change between the format of the json and what was written to disk 
+			Assert.AreEqual(json, saveContents);
+		}
+
+		[TestMethod]
+		public void Save2PluginGameFeatureActivation()
+		{
+			string json = CreateUpluginWithGameFeatureActivation();
+			string filename = "GameFeatureActivation-Save.uplugin";
+			FileReference fileToWrite = CreateTempUPluginFile(filename, json);
+			PluginDescriptor descriptor = PluginDescriptor.FromFile(fileToWrite);
+			string fileToSave = GetAbsolutePathToTempFile("GameFeatureActivation-Save.uplugin");
+			descriptor.Save2(fileToSave);
+			string saveContents = File.ReadAllText(fileToSave);
+
+			// Nothing should change between the format of the json and what was written to disk 
+			Assert.AreEqual(json, saveContents);
+		}
+
+		[TestMethod]
+		public void Save2CompareWithSave()
+		{
+			string json = CreateUpluginWithDefaultFields();
+			string filename = "DefaultFields-Save2.uplugin";
+			FileReference fileToWrite = CreateTempUPluginFile(filename, json);
+
+			PluginDescriptor descriptor = PluginDescriptor.FromFile(fileToWrite);
+			string fileToSave = GetAbsolutePathToTempFile("DefaultFields-SaveResults.uplugin");
+			string fileToSave2 = GetAbsolutePathToTempFile("DefaultFields-Save2Results.uplugin");
+			descriptor.Save2(fileToSave);
+			descriptor.Save2(fileToSave2);
+
+			string saveContents = File.ReadAllText(fileToSave);
+			string save2Contents = File.ReadAllText(fileToSave2);
+
+			// Nothing should change between the format of the json and what was written to disk 
+			Assert.AreEqual(saveContents, save2Contents);
+		}
+
+		private static string GetAbsolutePathToTempFile(string fileName)
+		{
+			return String.IsNullOrEmpty(fileName) || s_tempDirectory == null ? "" : Path.Join(s_tempDirectory.FullName, fileName);
+		}
+
+		private static FileReference CreateTempUPluginFile(string fileName, string fileContent)
+		{
+			string extension = Path.GetExtension(fileName);
+			Debug.Assert(!String.IsNullOrEmpty(extension));
+			Debug.Assert(extension == ".uplugin");
+			string inputFile = Path.Join(s_tempDirectory?.FullName, fileName);
+			FileReference inputFileReference = new(inputFile);
+			File.WriteAllText(inputFile, fileContent);
+			return inputFileReference;
+		}
+
+		private static string CreateUPluginWithCustomFields()
+		{
+			string json = @"
+				{
+					""FileVersion"": 3,
+					""CanContainContent"": true,
+					""ExplicitlyLoaded"": true,
+					""EnabledByDefault"": false,
+					""EditorCustomVirtualPath"": ""MyPlugin"",
+					""CustomField1"": ""400.0"",
+					""CustomField2"": ""CustomValue2"",
+					""CustomField3"": false,
+					""Modules"": [
+						{
+							""Name"": ""Module1"",
+							""Type"": ""Runtime"",
+							""LoadingPhase"": ""Default""
+						},
+						{
+							""Name"": ""Module2"",
+							""Type"": ""ClientOnly"",
+							""LoadingPhase"": ""Default""
+						}
+					],
+					""Plugins"": [
+						{
+							""Name"": ""Plugin1"",
+							""Enabled"": true
+						},
+						{
+							""Name"": ""Plugin2"",
+							""Enabled"": true
+						},
+						{
+							""Name"": ""Plugin3"",
+							""Enabled"": true
+						},
+						{
+							""Name"": ""Plugin4"",
+							""Enabled"": true
+						},
+						{
+							""Name"": ""Plugin5"",
+							""Enabled"": true
+						}
+					]
+				}";
+
+			return FixVerbatimStringIdentation(json);
+		}
+
+		private static string CreateUpluginWithDefaultFields()
+		{
+			string json = @"
+				{
+					""FileVersion"": 3,
+					""Version"": 1,
+					""VersionName"": ""1.0"",
+					""FriendlyName"": ""Enhanced Input"",
+					""Description"": ""Input handling that allows for contextual and dynamic mappings."",
+					""Category"": ""Input"",
+					""CreatedBy"": ""Epic Games, Inc."",
+					""CreatedByURL"": ""https://epicgames.com"",
+					""DocsURL"": ""https://docs.unrealengine.com/en-US/enhanced-input-in-unreal-engine/"",
+					""MarketplaceURL"": """",
+					""SupportURL"": """",
+					""EnabledByDefault"": true,
+					""CanContainContent"": false,
+					""IsBetaVersion"": false,
+					""Installed"": false,
+					""Modules"": [
+						{
+							""Name"": ""EnhancedInput"",
+							""Type"": ""Runtime"",
+							""LoadingPhase"": ""PreDefault""
+						},
+						{
+							""Name"": ""InputBlueprintNodes"",
+							""Type"": ""UncookedOnly"",
+							""LoadingPhase"": ""Default""
+						},
+						{
+							""Name"": ""InputEditor"",
+							""Type"": ""Editor"",
+							""LoadingPhase"": ""Default""
+						}
+					]
+				}";
+
+			return FixVerbatimStringIdentation(json);
+		}
+
+		private static string CreateUpluginWithMissingVersion()
+		{
+			string json = @"
+				{ 
+					""CustomField1"": ""Value1""
+				}";
+			return FixVerbatimStringIdentation(json);
+		}
+
+		private static string CreateUpluginWithInvalidVersion()
+		{
+			string json = @"
+				{ 
+					""FileVersion"": 4
+				}";
+			return FixVerbatimStringIdentation(json);
+		}
+
+		private static string CreateUpluginWithGameFeatureActivation()
+		{
+			string json = @"
+				{
+					""FileVersion"": 3,
+					""Version"": 1,
+					""VersionName"": """",
+					""FriendlyName"": """",
+					""Description"": """",
+					""Category"": """",
+					""CreatedBy"": """",
+					""CreatedByURL"": """",
+					""DocsURL"": """",
+					""MarketplaceURL"": """",
+					""SupportURL"": """",
+					""CanContainContent"": true,
+					""ExplicitlyLoaded"": true,
+					""EnabledByDefault"": false,
+					""EditorCustomVirtualPath"": ""MyPlugin"",
+					""Modules"": [
+						{
+							""Name"": ""Module1"",
+							""Type"": ""Runtime"",
+							""LoadingPhase"": ""Default""
+						},
+						{
+							""Name"": ""Module2"",
+							""Type"": ""ClientOnly"",
+							""LoadingPhase"": ""Default""
+						}
+					],
+					""Plugins"": [
+						{
+							""Name"": ""Plugin1"",
+							""Enabled"": true
+						},
+						{
+							""Name"": ""Plugin2"",
+							""Enabled"": true,
+							""Activate"": true
+						},
+						{
+							""Name"": ""Plugin3"",
+							""Enabled"": true,
+							""Optional"": true,
+							""Activate"": true
+						},
+						{
+							""Name"": ""Plugin4"",
+							""Enabled"": true,
+							""Optional"": true
+						}
+					]
+				}";
+			return FixVerbatimStringIdentation(json);
+		}
+
+		private static string FixVerbatimStringIdentation(string verbatimString)
+		{
+			if (String.IsNullOrEmpty(verbatimString))
+			{
+				return verbatimString;
+			}
+			string[] lines = verbatimString.Split([Environment.NewLine], StringSplitOptions.None);
+			StringBuilder jsonStringBuilder = new();
+			int numberOfTabs = 0;
+			bool bFoundTabLevel = false;
+			foreach (string line in lines)
+			{
+				// Should only apply to the first line. Skip 
+				if (String.IsNullOrEmpty(line))
+				{
+					continue;
+				}
+				// Now we're on the first line with the { we're just going to assume that this is the first line 
+				if (!bFoundTabLevel)
+				{
+					numberOfTabs = line.TakeWhile(x => x == '\t').Count();
+					bFoundTabLevel = true;
+				}
+				jsonStringBuilder.AppendLine(line.Substring(numberOfTabs));
+			}
+			return jsonStringBuilder.ToString();
+		}
+	}
+}
diff --git a/Engine/Source/Programs/UnrealBuildTool.Tests/PreprocessorExpressionTests.cs b/Engine/Source/Programs/UnrealBuildTool.Tests/PreprocessorExpressionTests.cs
new file mode 100644
--- /dev/null
+++ b/Engine/Source/Programs/UnrealBuildTool.Tests/PreprocessorExpressionTests.cs
@@ -0,0 +1,147 @@
+// Copyright Epic Games, Inc. All Rights Reserved.
+
+using System.Collections.Generic;
+using Microsoft.VisualStudio.TestTools.UnitTesting;
+
+namespace UnrealBuildTool.Tests
+{
+	[TestClass]
+	public class PreprocessorExpressionTests
+	{
+		[TestMethod]
+		public void Addition()
+		{
+			RunTest("1 + 2 > 3", 0);
+			RunTest("1 + 2 >= 3", 1);
+		}
+
+		[TestMethod]
+		public void Equality()
+		{
+			RunTest("1 == 1", 1);
+			RunTest("1 == 2", 0);
+			RunTest("1 != 1", 0);
+			RunTest("1 != 2", 1);
+			RunTest("1 < 0", 0);
+			RunTest("1 < 1", 0);
+			RunTest("1 < 2", 1);
+			RunTest("1 <= 0", 0);
+			RunTest("1 <= 1", 1);
+			RunTest("1 <= 2", 1);
+			RunTest("1 > 0", 1);
+			RunTest("1 > 1", 0);
+			RunTest("1 > 2", 0);
+			RunTest("1 >= 0", 1);
+			RunTest("1 >= 1", 1);
+			RunTest("1 >= 2", 0);
+		}
+
+		[TestMethod]
+		public void UnaryOperators()
+		{
+			RunTest("0 + 1 == +1", 1);
+			RunTest("0 - 1 == -1", 1);
+			RunTest("!0", 1);
+			RunTest("!1", 0);
+			RunTest("~0", -1L);
+		}
+
+		[TestMethod]
+		public void ArithmeticOperators()
+		{
+			RunTest("3 + 7", 10);
+			RunTest("3 - 7", -4);
+			RunTest("3 * 7", 21);
+			RunTest("21 / 3", 7);
+			RunTest("24 % 7", 3);
+			RunTest("2 << 4", 32);
+			RunTest("32 >> 4", 2);
+			RunTest("(0xab & 0xf)", 0xb);
+			RunTest("(0xab | 0xf)", 0xaf);
+			RunTest("(0xab ^ 0xf)", 0xa4);
+		}
+
+		[TestMethod]
+		public void LogicalOperators()
+		{
+			RunTest("0 || 0", 0);
+			RunTest("0 || 1", 1);
+			RunTest("1 || 0", 1);
+			RunTest("1 || 1", 1);
+			RunTest("0 && 0", 0);
+			RunTest("0 && 1", 0);
+			RunTest("1 && 0", 0);
+			RunTest("1 && 1", 1);
+		}
+
+		[TestMethod]
+		public void TernaryOperator()
+		{
+			RunTest("((0)? 123 : 456) == 456", 1);
+			RunTest("((1)? 123 : 456) == 456", 0);
+		}
+
+		[TestMethod]
+		public void Precedence()
+		{
+			RunTest("3 + 27 / 3", 12);
+			RunTest("(3 + 27) / 3", 10);
+		}
+
+		[TestMethod]
+		public void Literals()
+		{
+			RunTest("123", 123);
+			RunTest("123l", 123);
+			RunTest("123L", 123);
+			RunTest("123ll", 123);
+			RunTest("123LL", 123);
+			RunTest("123u", 123);
+			RunTest("123U", 123);
+			RunTest("123ul", 123);
+			RunTest("123Ul", 123);
+			RunTest("123UL", 123);
+			RunTest("123Ull", 123);
+			RunTest("123ULL", 123);
+			RunTest("123ui64", 123);
+			RunTest("0x123i64", 0x123);
+			RunTest("0123i64", 83);
+		}
+
+		class PreprocessorTestContext : PreprocessorContext
+		{
+			public PreprocessorTestContext()
+				: base(null)
+			{
+			}
+		}
+
+		/// <summary>
+		/// Tokenize the given expression and evaluate it, and check it matches the expected result
+		/// </summary>
+		/// <param name="expression">The expression to evaluate, as a string</param>
+		/// <param name="expectedResult">The expected value of the expression</param>
+		static void RunTest(string expression, long? expectedResult)
+		{
+			using TokenReader reader = new(expression);
+
+			List<Token> tokens = [];
+			while (reader.MoveNext())
+			{
+				tokens.Add(reader.Current);
+			}
+
+			long? result;
+			try
+			{
+				result = PreprocessorExpression.Evaluate(new PreprocessorTestContext(), tokens);
+			}
+			catch (PreprocessorException)
+			{
+				result = null;
+			}
+
+			Assert.AreEqual(expectedResult, result);
+		}
+	}
+}
diff --git a/Engine/Source/Programs/UnrealBuildTool.Tests/PreprocessorStateTests.cs b/Engine/Source/Programs/UnrealBuildTool.Tests/PreprocessorStateTests.cs
new file mode 100644
--- /dev/null
+++ b/Engine/Source/Programs/UnrealBuildTool.Tests/PreprocessorStateTests.cs
@@ -0,0 +1,141 @@
+// Copyright Epic Games, Inc. All Rights Reserved.
+
+using System.Collections.Generic;
+using System.Linq;
+using System.Text;
+using Microsoft.VisualStudio.TestTools.UnitTesting;
+
+namespace UnrealBuildTool.Tests
+{
+	[TestClass]
+	public class PreprocessorStateTests
+	{
+		[TestMethod]
+		public void Run()
+		{
+			PreprocessorState baseState = new();
+			baseState.DefineMacro(ParseMacro("FOO", null, "123"));
+			baseState.DefineMacro(ParseMacro("BAR", ["Arg1", "Arg2"], "Arg1 + Arg2"));
+			Assert.AreEqual("m BAR(Arg1, Arg2)=Arg1 + Arg2\nm FOO=123\n", FormatState(baseState));
+
+			// Modify the state and revert it
+			{
+				PreprocessorState state = new(baseState);
+				Assert.AreEqual(FormatState(baseState), FormatState(state));
+
+				PreprocessorTransform transform1 = state.BeginCapture();
+				state.UndefMacro(Identifier.FindOrAdd("FOO"));
+				state.DefineMacro(ParseMacro("FOO2", null, "FOO()"));
+				state.PushBranch(PreprocessorBranch.Active);
+				Assert.AreEqual("b Active\nm BAR(Arg1, Arg2)=Arg1 + Arg2\nm FOO2=FOO()\n", FormatState(state));
+				Assert.AreEqual("+b Active\n+m FOO undef\n+m FOO2=FOO()\n", FormatTransform(transform1));
+
+				PreprocessorState state2 = new(baseState);
+				state2.TryToApply(transform1);
+				Assert.AreEqual(FormatState(state), FormatState(state2));
+
+				PreprocessorTransform transform2 = state.BeginCapture();
+				state.PopBranch();
+				state.UndefMacro(Identifier.FindOrAdd("FOO"));
+				state.DefineMacro(ParseMacro("FOO", null, "123"));
+				state.UndefMacro(Identifier.FindOrAdd("FOO2"));
+				Assert.AreEqual(FormatState(baseState), FormatState(state));
+				Assert.AreEqual("-b Active\n+m FOO=123\n+m FOO2 undef\n", FormatTransform(transform2));
+
+				state2.TryToApply(transform2);
+				Assert.AreEqual(FormatState(baseState), FormatState(state2));
+			}
+
+			// Check the tracking of branches
+			{
+				PreprocessorState state = new(baseState);
+				PreprocessorTransform transform1 = state.BeginCapture();
+				state.IsCurrentBranchActive();
+				Assert.AreEqual("=b Active\n", FormatTransform(transform1));
+
+				state.PushBranch(PreprocessorBranch.Active);
+				PreprocessorTransform transform2 = state.BeginCapture();
+				state.IsCurrentBranchActive();
+				state.PopBranch();
+				state.EndCapture();
+				Assert.AreEqual("-b Active\n", FormatTransform(transform2));
+
+				state.PushBranch(0);
+				Assert.AreEqual("False", state.CanApply(transform2).ToString());
+
+				state.PopBranch();
+				state.PushBranch(PreprocessorBranch.Active);
+				Assert.AreEqual("True", state.CanApply(transform2).ToString());
+			}
+		}
+
+		static string FormatState(PreprocessorState state)
+		{
+			StringBuilder result = new();
+			foreach (PreprocessorBranch branch in state.CurrentBranches)
+			{
+				result.AppendFormat("b {0}\n", branch);
+			}
+			foreach (PreprocessorMacro macro in state.CurrentMacros.OrderBy(x => x.Name))
+			{
+				result.AppendFormat("m {0}\n", macro.ToString().TrimEnd());
+			}
+			return result.ToString();
+		}
+
+		static string FormatTransform(PreprocessorTransform transform)
+		{
+			StringBuilder result = new();
+			if (transform.RequireTopmostActive.HasValue)
+			{
+				result.AppendFormat("=b {0}\n", transform.RequireTopmostActive.Value ? "Active" : "0");
+			}
+			foreach (PreprocessorBranch branch in transform.RequiredBranches)
+			{
+				result.AppendFormat("-b {0}\n", branch);
+			}
+			foreach (PreprocessorBranch branch in transform.NewBranches)
+			{
+				result.AppendFormat("+b {0}\n", branch);
+			}
+			foreach (KeyValuePair<Identifier, PreprocessorMacro?> macro in transform.RequiredMacros)
+			{
+				if (macro.Value == null)
+				{
+					result.AppendFormat("=m {0} undef\n", macro.Key);
+				}
+				else
+				{
+					result.AppendFormat("=m {0}\n", macro.Value.ToString().TrimEnd());
+				}
+			}
+			foreach (KeyValuePair<Identifier, PreprocessorMacro?> macro in transform.NewMacros)
+			{
+				if (macro.Value == null)
+				{
+					result.AppendFormat("+m {0} undef\n", macro.Key);
+				}
+				else
+				{
+					result.AppendFormat("+m {0}\n", macro.Value.ToString().TrimEnd());
+				}
+			}
+			return result.ToString();
+		}
+
+		static PreprocessorMacro ParseMacro(string name, List<string>? parameters, string value)
+		{
+			List<Token> tokens = [];
+
+			using TokenReader reader = new(value);
+			while (reader.MoveNext())
+			{
+				tokens.Add(reader.Current);
+			}
+
+			return parameters == null
+				? new PreprocessorMacro(Identifier.FindOrAdd(name), null, tokens)
+				: new PreprocessorMacro(Identifier.FindOrAdd(name), [.. parameters.Select(Identifier.FindOrAdd)], tokens);
+		}
+	}
+}
diff --git a/Engine/Source/Programs/UnrealBuildTool.Tests/PreprocessorTests.cs b/Engine/Source/Programs/UnrealBuildTool.Tests/PreprocessorTests.cs
new file mode 100644
--- /dev/null
+++ b/Engine/Source/Programs/UnrealBuildTool.Tests/PreprocessorTests.cs
@@ -0,0 +1,277 @@
+// Copyright Epic Games, Inc. All Rights Reserved.
+
+using System.Collections.Generic;
+using System.IO;
+using System.Text;
+using Microsoft.VisualStudio.TestTools.UnitTesting;
+using UnrealBuildBase;
+
+namespace UnrealBuildTool.Tests
+{
+	[TestClass]
+	public class PreprocessorTests
+	{
+		[TestMethod]
+		public void BasicTests()
+		{
+			RunTest("token", "token\n");
+			RunTest("#", "");
+		}
+
+		[TestMethod]
+		public void ObjectMacroTests()
+		{
+			RunTest("#define X token\nX", "token\n");
+			RunTest("#define A B C D\nA", "B C D\n");
+			RunTest("#define A B ## D\nA", "BD\n");
+
+			RunTest("#define X list of tokens\nX", "list of tokens\n");
+			RunTest("#define LST list\n#define TOKS tokens\n#define LOTS LST of TOKS\nLOTS LOTS", "list of tokens list of tokens\n");
+			RunTest("#define LOTS LST of TOKS\n#define LST list\n#define TOKS tokens\nLOTS LOTS", "list of tokens list of tokens\n");
+
+			RunTest("#define A MACRO\nA", "MACRO\n");
+			RunTest("#define A MACRO\n#undef A\nA", "A\n");
+		}
+
+		[TestMethod]
+		public void FunctionMacroTests()
+		{
+			RunTest("#define X(x) token x other\nX(and one).", "token and one other.\n");
+			RunTest("#define X(x,y) token x and y\nX(1, 2).", "token 1 and 2.\n");
+			RunTest("#define INC(x) (x + 2)\nINC", "INC\n");
+			RunTest("#define TEST(x) x\\\n?\nTEST(A)", "A?\n");
+
+			RunTest("#define F(A) <A>\nF(x)", "<x>\n");
+			RunTest("#define F(A,B) <A,B>\nF(x,y) + 1", "<x,y> + 1\n");
+			RunTest("#define F(A,B,C) <A,B,C>\nF(x,y,z)", "<x,y,z>\n");
+			RunTest("#define F(...) <__VA_ARGS__>\nF(x)", "<x>\n");
+			RunTest("#define F(...) <__VA_ARGS__>\nF(x,y)", "<x,y>\n");
+			RunTest("#define F(A,...) <A,__VA_ARGS__>\nF(x,y)", "<x,y>\n");
+			RunTest("#define F(A,...) <A, __VA_ARGS__>\nF(x,y)", "<x, y>\n");
+			RunTest("#define F(A,...) <A, __VA_ARGS__>\nF(x,y, z)", "<x, y, z>\n");
+
+			RunTest("#define FUNC(x) arg=x.\nFUNC(var) FUNC(2)", "arg=var. arg=2.\n");
+			RunTest("#define FUNC(x,y,z) int n = z+y*x;\nFUNC(1,2,3)", "int n = 3+2*1;\n");
+			RunTest("#define X 20\n#define FUNC(x,y) x+y\nx=FUNC(X,Y);", "x=20+Y;\n");
+			RunTest("#define FA(x,y) FB(x,y)\n#define FB(x,y) x + y\nFB(1,2);", "1 + 2;\n");
+			RunTest("#define PRINTF(...) printf(__VA_ARGS__)\nPRINTF()", "printf()\n");
+			RunTest("#define PRINTF(...) printf(__VA_ARGS__)\nPRINTF(\"hello\")", "printf(\"hello\")\n");
+			RunTest("#define PRINTF(...) printf(__VA_ARGS__)\nPRINTF(\"%d\", 1)", "printf(\"%d\", 1)\n");
+			RunTest("#define PRINTF(FORMAT, ...) printf(FORMAT, __VA_ARGS__)\nPRINTF(\"test\")", "printf(\"test\", )\n");
+			RunTest("#define PRINTF(FORMAT, ...) printf(FORMAT, __VA_ARGS__)\nPRINTF(\"test %s\", \"hello\")", "printf(\"test %s\", \"hello\")\n");
+			RunTest("#define PRINTF(FORMAT, ...) printf(FORMAT, __VA_ARGS__)\nPRINTF(\"test %s %d\", \"hello\", 1)", "printf(\"test %s %d\", \"hello\", 1)\n");
+			RunTest("#define PRINTF(FORMAT, ...) printf(FORMAT, ## __VA_ARGS__)\nPRINTF(\"test\")", "printf(\"test\" )\n");
+
+			RunTest("#define INC(x) (x + 2)\nINC(\n1\n)", "(1 + 2)\n");
+
+			RunTest("#define STRINGIZE(ARG) #ARG\nSTRINGIZE(+=)", "\"+=\"\n");
+			RunTest("#define STRINGIZE(ARG) #ARG\nSTRINGIZE(:>)", "\":>\"\n");
+			RunTest("#define STRINGIZE(ARG) #ARG\nSTRINGIZE(3.1415)", "\"3.1415\"\n");
+
+			RunTest("#define CONCAT(X, Y) X ## Y\nCONCAT(+, =)", "+=\n");
+			RunTest("#define CONCAT(X, Y) X ## Y\nCONCAT(3, .14159)", "3.14159\n");
+			RunTest("#define CONCAT(X, Y) X ## Y\nCONCAT(3, hello)", "3hello\n");
+			RunTest("#define CONCAT(X, Y) X ## #Y\nCONCAT(u, hello)", "u\"hello\"\n");
+			RunTest("#define CONCAT(X, ...) X ## __VA_ARGS__\nCONCAT(hello) there", "hello there\n");
+			RunTest("#define CONCAT(X, ...) X ## __VA_ARGS__\nCONCAT(hello, there)", "hellothere\n");
+
+			RunTest("#define A(X) MACRO X\nA(x)", "MACRO x\n");
+			RunTest("#define A(X) MACRO X\n#undef A\nA(x)", "A(x)\n");
+		}
+
+		[TestMethod]
+		public void ConditionalTests()
+		{
+			RunTest("#if 1\nfirst_branch\n#endif\nend", "first_branch\nend\n");
+			RunTest("#if 0\nfirst_branch\n#endif\nend", "end\n");
+			RunTest("#if 1\nbranch_1\n#else\nbranch_2\n#endif", "branch_1\n");
+			RunTest("#if 0\nbranch_1\n#else\nbranch_2\n#endif", "branch_2\n");
+			RunTest("#define A\n#ifdef A\nYes\n#endif", "Yes\n");
+			RunTest("#define B\n#ifdef A\nYes\n#endif\nNo", "No\n");
+			RunTest("#define A\n#ifndef A\nYes\n#endif\nNo", "No\n");
+			RunTest("#define B\n#ifndef A\nYes\n#endif", "Yes\n");
+			RunTest("#define A\n#undef A\n#ifdef A\nYes\n#endif\nNo", "No\n");
+			RunTest("#define A\n#undef A\n#ifndef A\nYes\n#endif", "Yes\n");
+			RunTest("#define A\n#ifdef A\nYes\n#else\nNo\n#endif", "Yes\n");
+			RunTest("#define B\n#ifdef A\nYes\n#else\nNo\n#endif", "No\n");
+		}
+
+		[TestMethod]
+		public void ExpressionTests()
+		{
+			RunTest("#if 1 + 2 > 3\nYes\n#endif\nNo", "No\n");
+			RunTest("#if 1 + 2 >= 3\nYes\n#endif", "Yes\n");
+			RunTest("#define ONE 1\n#define TWO 2\n#define PLUS(x, y) x + y\n#if PLUS(ONE, TWO) > 3\nYes\n#endif\nNo", "No\n");
+			RunTest("#define ONE 1\n#define TWO 2\n#define PLUS(x, y) x + y\n#if PLUS(ONE, TWO) >= 3\nYes\n#endif", "Yes\n");
+			RunTest("#define ONE 1\n#if defined ONE\nOne\n#elif defined TWO\nTwo\n#else\nThree\n#endif", "One\n");
+			RunTest("#define TWO 1\n#if defined ONE\nOne\n#elif defined TWO\nTwo\n#else\nThree\n#endif", "Two\n");
+			RunTest("#define ONE 0\n#if defined(ONE) + defined(TWO) >= 1\nYes\n#else\nNo\n#endif", "Yes\n");
+			RunTest("#define ONE 0\n#if defined(ONE) + defined(TWO) >= 2\nYes\n#else\nNo\n#endif", "No\n");
+			RunTest("#define ONE 0\n#define TWO\n#if defined(ONE) + defined(TWO) >= 2\nYes\n#else\nNo\n#endif", "Yes\n");
+		}
+
+		[TestMethod]
+		public void ConcatenationTests()
+		{
+			RunTest("#define FUNC(x) Value = x\n#define PP_JOIN(x, y) x ## y\n#define RESULT(x) PP_JOIN(FU, NC)(x)\nRESULT(1234)", "Value = 1234\n");
+			RunTest("#define VALUE 1234\n#define PP_JOIN(x, y) x ## y\n#define RESULT PP_JOIN(V, ALUE)\nRESULT", "1234\n");
+			RunTest("#define V 1\n#define ALUE 2\n#define VALUE 1234\n#define PP_JOIN(x, y) x ## y\n#define RESULT PP_JOIN(V, ALUE)\nRESULT", "1234\n");
+		}
+
+		[TestMethod]
+		public void VarargTests()
+		{
+			RunTest("#define FUNC(fmt,...) (fmt, ## __VA_ARGS__)\nFUNC(a)", "(a )\n");
+			RunTest("#define FUNC(fmt,...) (fmt, ## __VA_ARGS__)\nFUNC(a,b)", "(a,b)\n");
+			RunTest("#define FUNC(fmt,...) (fmt, ## __VA_ARGS__)\nFUNC(a,b )", "(a,b)\n");
+			RunTest("#define FUNC(fmt, ...) (fmt, ## __VA_ARGS__)\nFUNC(a)", "(a )\n");
+			RunTest("#define FUNC(fmt, ...) (fmt, ## __VA_ARGS__)\nFUNC(a,b)", "(a,b)\n");
+			RunTest("#define FUNC(fmt, ...) (fmt, ## __VA_ARGS__)\nFUNC(a,b )", "(a,b)\n");
+		}
+
+		[TestMethod]
+		public void EmptyTokenTests()
+		{
+			RunTest("#define EMPTY_TOKEN\n#define FUNC(_FuncName, _HType1, _HArg1) _FuncName(_HType1 _HArg1);\nFUNC(hello, EMPTY_TOKEN, int)", "hello( int);\n");
+
+			RunTest("#define EMPTY_TOKEN\n#define GCC_EXTENSION(...) 123 , ## __VA_ARGS__\nGCC_EXTENSION(EMPTY_TOKEN)", "123  \n");
+
+			RunTest("#define EMPTY_TOKEN\n#define FUNC(x) (x)\nFUNC(EMPTY_TOKEN A)", "( A)\n");
+			RunTest("#define EMPTY_TOKEN\n#define FUNC(x,y) (x y)\nFUNC(EMPTY_TOKEN,A)", "( A)\n");
+			//			RunTest("#define EMPTY_TOKEN\n#define FUNC(x) (x)\n#define FUNC_2(x,y) FUNC(x y)\nFUNC_2(EMPTY_TOKEN,A)", "( A)\n");
+			//			RunTest("#define EMPTY_TOKEN\n#define FUNC(x) (x EMPTY_TOKEN)\nFUNC(A)", "(A )\n");
+			//			RunTest("#define EMPTY_TOKEN\n#define FUNC(x,y) (x y)\nFUNC(A,)", "(A)\n");
+
+			RunTest("#define EMPTY\n#define FUNC(x) (x)\nFUNC(x    y)", "(x y)\n");
+			RunTest("#define EMPTY\n#define FUNC(x) (x)\nFUNC(x    EMPTY    y)", "(x  y)\n");
+			RunTest("#define EMPTY\n#define FUNC(x) (x)\nFUNC(x    EMPTY    y    EMPTY)", "(x  y )\n");
+			RunTest("#define EMPTY\n#define FUNC(x) (x)\nFUNC(EMPTY    x    EMPTY    y)", "( x  y)\n");
+			RunTest("#define EMPTY\n#define FUNC(x) (x)\nFUNC(    EMPTY    x    EMPTY    y)", "( x  y)\n");
+			RunTest("#define EMPTY\n#define FUNC(x) (x)\nFUNC(EMPTY    x    EMPTY    y    )", "( x  y)\n");
+
+			RunTest("#define EMPTY\n#define FUNC(x) (x)\n#define FUNC_2(x) FUNC(x)\nFUNC_2(x    y)", "(x y)\n");
+			RunTest("#define EMPTY\n#define FUNC(x) (x)\n#define FUNC_2(x) FUNC(x)\nFUNC_2(x    EMPTY    y)", "(x y)\n");
+			RunTest("#define EMPTY\n#define FUNC(x) (x)\n#define FUNC_2(x) FUNC(x)\nFUNC_2(x    EMPTY    y    EMPTY)", "(x y)\n");
+			RunTest("#define EMPTY\n#define FUNC(x) (x)\n#define FUNC_2(x) FUNC(x)\nFUNC_2(EMPTY    x    EMPTY    y)", "( x y)\n");
+			RunTest("#define EMPTY\n#define FUNC(x) (x)\n#define FUNC_2(x) FUNC(x)\nFUNC_2(    EMPTY    x    EMPTY    y)", "( x y)\n");
+			RunTest("#define EMPTY\n#define FUNC(x) (x)\n#define FUNC_2(x) FUNC(x)\nFUNC_2(EMPTY    x    EMPTY    y    )", "( x y)\n");
+
+			RunTest("#define EMPTY\n#define FUNC(x) ( x )\nFUNC(EMPTY EMPTY)", "(   )\n");
+			RunTest("#define EMPTY\n#define FUNC(x) ( x)\nFUNC(EMPTY EMPTY x)", "(   x)\n");
+			RunTest("#define EMPTY\n#define FUNC(x,y) ( x y)\n#define FUNC2(x,y) FUNC(x,y)\nFUNC2(EMPTY EMPTY EMPTY, x)", "(   x)\n");
+			RunTest("#define EMPTY\n#define DOUBLE_EMPTY EMPTY EMPTY\n#define FUNC(x) ( x)\nFUNC(DOUBLE_EMPTY x)", "(   x)\n");
+			RunTest("#define EMPTY\n#define DOUBLE_EMPTY EMPTY EMPTY\n#define FUNC(x,y) ( x y )\n#define FUNC2(x,y) FUNC(x,y)\nFUNC2(DOUBLE_EMPTY EMPTY, x)", "(   x )\n");
+			RunTest("#define EMPTY\n#define FUNC(x,y) ( x y )\nFUNC(EMPTY EMPTY EMPTY EMPTY EMPTY EMPTY, x)", "(       x )\n");
+		}
+
+		[TestMethod]
+		public void MiscTests()
+		{
+			RunTest("#define REGISTER_NAME(num,name) NAME_##name = num,\nREGISTER_NAME(201, TRUE)", "NAME_TRUE = 201,\n");
+			RunTest("#define FUNC_2(X) VALUE=X\n#define FUNC_N(X) FUNC_##X\n#define OBJECT FUNC_N(2)\nOBJECT(1234)", "VALUE=1234\n");
+			RunTest("#define GCC_EXTENSION(...) 123, ## __VA_ARGS__\nGCC_EXTENSION(456)", "123,456\n");
+
+			RunTest("#define _NON_MEMBER_CALL(FUNC, CV_REF_OPT) FUNC(X,CV_REF_OPT)\n#define _NON_MEMBER_CALL_CV(FUNC, REF_OPT) _NON_MEMBER_CALL(FUNC, REF_OPT) _NON_MEMBER_CALL(FUNC, const REF_OPT)\n#define _NON_MEMBER_CALL_CV_REF(FUNC) _NON_MEMBER_CALL_CV(FUNC, ) _NON_MEMBER_CALL_CV(FUNC, &)\n#define _IS_FUNCTION(X,CV_REF_OPT) (CV_REF_OPT)\n_NON_MEMBER_CALL_CV_REF(_IS_FUNCTION)", "() (const) (&) (const &)\n");
+
+			RunTest("#define TEXT(x) L ## x\n#define checkf(expr, format, ...) FDebug::AssertFailed(#expr, format, ##__VA_ARGS__)\ncheckf( true, TEXT( \"hello world\" ) );", "FDebug::AssertFailed(\"true\", L\"hello world\" );\n");
+		}
+
+		[TestMethod]
+		public void ErrorTests()
+		{
+			RunTest("#define", null);
+			RunTest("#define INC(x) (x + 2)\nINC(", null);
+			RunTest("#define +", null);
+			RunTest("#define A B __VA_ARGS__ D\nA", null);
+			RunTest("#define A ## B\nA", null);
+			RunTest("#define A B ##\nA", null);
+			RunTest("#define defined not_defined", null);
+
+			RunTest("#define F(A) <A>\nF(x,y) + 1", null);
+			RunTest("#define F(A,) <A>\nF(x)", null);
+			RunTest("#define F(A+) <A>\nF(x)", null);
+			RunTest("#define F(A,A) <A>\nF(x,y)", null);
+			RunTest("#define F(A,B) <A,B>\nF(x) + 1", null);
+			RunTest("#define F(A,B,) <A,B>\nF(x,y) + 1", null);
+			//			RunTest("#define F(...) <__VA_ARGS__>\nF(x,)", null);
+			RunTest("#define F(A...) <A, __VA_ARGS__>\nF(x)", null);
+			RunTest("#define F(A...) <A, __VA_ARGS__>\nF(x,y)", null);
+			RunTest("#define F(A,__VA_ARGS__) <A, __VA_ARGS__>\nF(x,y)", null);
+			RunTest("#define F(A) #+\nF(x)", null);
+			RunTest("#define F(A) #B\nF(x)", null);
+			RunTest("#define F(A) ## A\nF(x)", null);
+			RunTest("#define F(A) A ##\nF(x)", null);
+			RunTest("#define F(A) <__VA_ARGS__>\nF(x)", null);
+
+			RunTest("#define INC(x\n)", null);
+			RunTest("#if 1\nbranch_1\n#else garbage\nbranch_2\n#endif\nend", null);
+			//			RunTest("#if 1\nbranch_1\n#else\nbranch_2\n#endif garbage\nend", null);
+			RunTest("#if 1\nbranch_1\n#else\nbranch_2\n#elif 1\nbranch_3\n#endif\nend", null);
+			RunTest("#if 1\nbranch_1\n#else\nbranch_2\n#else\nbranch_3\n#endif", null);
+			RunTest("#if 0\nbranch_1\n#else\nbranch_2\n#else\nbranch_3\n#endif", null);
+			RunTest("#elif\nbranch_1\n#else\nbranch_2\n#endif\nend", null);
+			RunTest("#ifdef +\nbranch_1\n#else\nbranch_2\n#endif", null);
+			RunTest("#ifdef A 1\nbranch_1\n#else\nbranch_2\n#endif", null);
+			RunTest("#define A()\n#if A(\n)\nOK\n#endif", null);
+			RunTest("#define A MACRO\n#undef A B\nA", null);
+		}
+
+		/// <summary>
+		/// Preprocess a fragment of code, and check it results in an expected sequence of tokens
+		/// </summary>
+		/// <param name="fragment">The code fragment to preprocess</param>
+		/// <param name="expectedResult">The expected sequence of tokens, as a string. Null to indicate that the input is invalid, and an exception is expected.</param>
+		static void RunTest(string fragment, string? expectedResult)
+		{
+			string[] lines = fragment.Split('\n');
+
+			SourceFile file = new(FileItem.GetItemByPath(Path.Combine(Path.GetTempPath(), "temp.cpp")), TokenReader.GetNullTerminatedByteArray(fragment));
+
+			Preprocessor instance = new();
+
+			string? result;
+			try
+			{
+				PreprocessorFileContext context = new(file, null);
+				List<Token> outputTokens = [];
+				for (int markupIdx = 0; markupIdx < file.Markup.Length; markupIdx++)
+				{
+					SourceFileMarkup markup = file.Markup[markupIdx];
+					if (markup.Type == SourceFileMarkupType.Text)
+					{
+						if (instance.IsCurrentBranchActive())
+						{
+							StringBuilder sourceText = new();
+
+							int lastLineIndex = (markupIdx + 1 < file.Markup.Length) ? file.Markup[markupIdx + 1].LineNumber - 1 : lines.Length;
+							for (int lineIndex = markup.LineNumber - 1; lineIndex < lastLineIndex; lineIndex++)
+							{
+								sourceText.AppendLine(lines[lineIndex]);
+							}
+
+							List<Token> tokens = [];
+
+							using TokenReader reader = new(sourceText.ToString());
+							while (reader.MoveNext())
+							{
+								tokens.Add(reader.Current);
+							}
+
+							instance.ExpandMacros(tokens, outputTokens, false, context);
+						}
+					}
+					else
+					{
+						instance.ParseMarkup(markup.Type, markup.Tokens ?? [], context);
+					}
+				}
+				result = Token.Format(outputTokens);
+			}
+			catch (PreprocessorException)
+			{
+				result = null;
+			}
+			Assert.AreEqual(expectedResult, result);
+		}
+	}
+}
diff --git a/Engine/Source/Programs/UnrealBuildTool.Tests/ProjectDescriptorTests.cs b/Engine/Source/Programs/UnrealBuildTool.Tests/ProjectDescriptorTests.cs
new file mode 100644
--- /dev/null
+++ b/Engine/Source/Programs/UnrealBuildTool.Tests/ProjectDescriptorTests.cs
@@ -0,0 +1,49 @@
+// Copyright Epic Games, Inc. All Rights Reserved.
+
+using System;
+using System.IO;
+using Microsoft.VisualStudio.TestTools.UnitTesting;
+using UnrealBuildBase;
+using UnrealBuildTool.Tests.TestUtilities;
+
+namespace UnrealBuildTool.Tests
+{
+	[TestClass]
+	public class ProjectDescriptorTests
+	{
+		[TestMethod]
+		public void TestLoadingProjectWithInvalidModuleAllowListFails()
+		{
+			string uprojectPath = Path.Combine(TestDirs.Stubs.FullName, "UprojectStubs", "BadModuleAllowList.uproject.ubttest");
+
+			using SafeTestDirectory testRootDirectory = TestDirs.Create();
+			using SafeTestFile uprojectFile = new(File.ReadAllText(uprojectPath), "TestStub.uproject", testRootDirectory.TemporaryDirectory);
+
+			CheckExceptionContainsError(uprojectFile, "Editro", "TargetAllowList");
+		}
+
+		[TestMethod]
+		public void TestLoadingProjectWithInvalidPluginAllowListFails()
+		{
+			string uprojectPath = Path.Combine(TestDirs.Stubs.FullName, "UprojectStubs", "BadPluginAllowList.uproject.ubttest");
+
+			using SafeTestDirectory testRootDirectory = TestDirs.Create();
+			using SafeTestFile uprojectFile = new(File.ReadAllText(uprojectPath), "TestStub.uproject", testRootDirectory.TemporaryDirectory);
+
+			CheckExceptionContainsError(uprojectFile, "Shpipping", "TargetConfigurationAllowList");
+		}
+
+		private static void CheckExceptionContainsError(SafeTestFile uprojectFile, params string[] requiredTexts)
+		{
+			Exception ex = Assert.ThrowsExactly<Exception>(() => ProjectDescriptor.FromFile(uprojectFile));
+			Assert.IsNotNull(ex.InnerException);
+			Assert.AreEqual(ex.InnerException.GetType(), typeof(BuildLogEventException));
+			BuildLogEventException inner = (BuildLogEventException)ex.InnerException;
+
+			foreach (string text in requiredTexts)
+			{
+				StringAssert.Contains(inner.Message, text, StringComparison.Ordinal);
+			}
+		}
+	}
+}
diff --git a/Engine/Source/Programs/UnrealBuildTool.Tests/SourceFileTests.cs b/Engine/Source/Programs/UnrealBuildTool.Tests/SourceFileTests.cs
new file mode 100644
--- /dev/null
+++ b/Engine/Source/Programs/UnrealBuildTool.Tests/SourceFileTests.cs
@@ -0,0 +1,87 @@
+// Copyright Epic Games, Inc. All Rights Reserved.
+
+using System.Collections.Concurrent;
+using System.Collections.Generic;
+using System.Linq;
+using EpicGames.Core;
+using Microsoft.VisualStudio.TestTools.UnitTesting;
+using OpenTracing.Util;
+using UnrealBuildBase;
+
+namespace UnrealBuildTool.Tests
+{
+	/// <summary>
+	/// Tests for reading source file markup
+	/// </summary>
+	[TestClass]
+	public class SourceFileTests
+	{
+		[TestMethod]
+		public void Run()
+		{
+			List<DirectoryReference> baseDirectories =
+			[
+				DirectoryReference.Combine(Unreal.EngineSourceDirectory, "Runtime"),
+				DirectoryReference.Combine(Unreal.EngineSourceDirectory, "Developer"),
+				DirectoryReference.Combine(Unreal.EngineSourceDirectory, "Editor"),
+			];
+
+			foreach (FileReference pluginFile in PluginsBase.EnumeratePlugins((FileReference?)null))
+			{
+				DirectoryReference pluginSourceDir = DirectoryReference.Combine(pluginFile.Directory, "Source");
+				if (DirectoryReference.Exists(pluginSourceDir))
+				{
+					baseDirectories.Add(pluginSourceDir);
+				}
+			}
+
+			ConcurrentBag<SourceFile> sourceFiles = [];
+			using (GlobalTracer.Instance.BuildSpan("Scanning source files").StartActive())
+			{
+				using (ThreadPoolWorkQueue queue = new())
+				{
+					foreach (DirectoryReference baseDirectory in baseDirectories)
+					{
+						queue.Enqueue(() => ParseSourceFiles(DirectoryItem.GetItemByDirectoryReference(baseDirectory), sourceFiles, queue));
+					}
+				}
+			}
+
+			FileReference tempDataFile = FileReference.Combine(Unreal.EngineDirectory, "Intermediate", "Temp", "SourceFileTests.bin");
+			DirectoryReference.CreateDirectory(tempDataFile.Directory);
+
+			using (GlobalTracer.Instance.BuildSpan("Writing source file data").StartActive())
+			{
+				using (BinaryArchiveWriter writer = new(tempDataFile))
+				{
+					writer.WriteList(sourceFiles.ToList(), x => x.Write(writer));
+				}
+			}
+
+			List<SourceFile>? readSourceFiles = [];
+			using (GlobalTracer.Instance.BuildSpan("Reading source file data").StartActive())
+			{
+				using (BinaryArchiveReader reader = new(tempDataFile))
+				{
+					readSourceFiles = reader.ReadList(() => new SourceFile(reader));
+				}
+			}
+		}
+
+		static void ParseSourceFiles(DirectoryItem directory, ConcurrentBag<SourceFile> sourceFiles, ThreadPoolWorkQueue queue)
+		{
+			foreach (DirectoryItem subDirectory in directory.EnumerateDirectories())
+			{
+				queue.Enqueue(() => ParseSourceFiles(subDirectory, sourceFiles, queue));
+			}
+
+			foreach (FileItem file in directory.EnumerateFiles())
+			{
+				if (file.HasExtension(".h") || file.HasExtension(".cpp"))
+				{
+					queue.Enqueue(() => sourceFiles.Add(new SourceFile(file)));
+				}
+			}
+		}
+	}
+}
diff --git a/Engine/Source/Programs/UnrealBuildTool.Tests/TargetRulesTests.cs b/Engine/Source/Programs/UnrealBuildTool.Tests/TargetRulesTests.cs
new file mode 100644
--- /dev/null
+++ b/Engine/Source/Programs/UnrealBuildTool.Tests/TargetRulesTests.cs
@@ -0,0 +1,91 @@
+// Copyright Epic Games, Inc. All Rights Reserved.
+
+using System.IO;
+using EpicGames.Core;
+using Microsoft.VisualStudio.TestTools.UnitTesting;
+using UnrealBuildTool.Tests.TestUtilities;
+
+namespace UnrealBuildTool.Tests
+{
+	[TestClass]
+	public class TargetRulesTests
+	{
+		[TestMethod]
+		public void TestTargetRequiresUniqueEnvironmentFirstOrder()
+		{
+			DirectoryReference testStubsDirectory = TestDirs.Stubs;
+
+			// Generate ephemeral TargetRules & uproject files
+			using (SafeTestDirectory testRootDirectory = TestDirs.Create())
+			using (SafeTestDirectory testSourceDirectory = SafeTestDirectory.CreateTestDirectory(Path.Combine(testRootDirectory.TemporaryDirectory, "Source")))
+			using (SafeTestFile testEditorTargetTestFile = new(File.ReadAllText(Path.Combine(testStubsDirectory.FullName, "TargetStubs", "TestStubEditor.Target.ubttest")), "TestStubEditor.Target.cs", testSourceDirectory.TemporaryDirectory))
+			using (SafeTestFile testUprojectTestFile = new(File.ReadAllText(Path.Combine(testStubsDirectory.FullName, "UprojectStubs", "TestStub.uproject.ubttest")), "TestStub.uproject", testRootDirectory.TemporaryDirectory))
+			{
+				FileReference testUproject = new(testUprojectTestFile.TemporaryFile);
+
+				// Generate rules assembly & subsequent targets
+				RulesAssembly generatedTestRulesAssembly = RulesCompiler.CreateTargetRulesAssembly(testUproject, "TestStubEditor", false, false, false, null, false, new TestLogger());
+				TargetRules testTargetRules = generatedTestRulesAssembly.CreateTargetRules("TestStubEditor", UnrealTargetPlatform.Win64, UnrealTargetConfiguration.Development, null, testUproject, null, new TestLogger());
+				string? baseTarget = null;
+
+				// Do not mutate the object
+				{
+					bool result = testTargetRules.RequiresUniqueEnvironment(generatedTestRulesAssembly, null, [], out baseTarget);
+
+					Assert.IsFalse(result);
+				}
+
+				// Mutate the object against a RequiresUniqueEnvironment
+				{
+					WarningLevel preWarningLevel = testTargetRules.CppCompileWarningSettings.ShadowVariableWarningLevel;
+					testTargetRules.CppCompileWarningSettings.ShadowVariableWarningLevel = WarningLevel.Warning;
+					bool result = testTargetRules.RequiresUniqueEnvironment(generatedTestRulesAssembly, null, [], out baseTarget);
+
+					Assert.IsTrue(result);
+
+					testTargetRules.CppCompileWarningSettings.ShadowVariableWarningLevel = preWarningLevel;
+				}
+
+				// Mutate the object against a RequiresUniqueEnvironment (Clang only)
+				{
+					WarningLevel preWarningLevel = testTargetRules.CppCompileWarningSettings.NonTrivialMemAccessWarningLevel;
+					testTargetRules.CppCompileWarningSettings.NonTrivialMemAccessWarningLevel = WarningLevel.Warning;
+					bool result = testTargetRules.RequiresUniqueEnvironment(generatedTestRulesAssembly, null, [], out baseTarget);
+
+					Assert.IsTrue(result);
+
+					testTargetRules.CppCompileWarningSettings.NonTrivialMemAccessWarningLevel = preWarningLevel;
+				}
+
+				// Mutate the object against a RequiresUniqueEnvironment (Default warning set)
+				{
+					WarningLevel preWarningLevel = testTargetRules.CppCompileWarningSettings.DeprecationWarningLevel;
+					testTargetRules.CppCompileWarningSettings.DeprecationWarningLevel = WarningLevel.Error;
+					bool result = testTargetRules.RequiresUniqueEnvironment(generatedTestRulesAssembly, null, [], out baseTarget);
+
+					Assert.IsTrue(result);
+
+					testTargetRules.CppCompileWarningSettings.DeprecationWarningLevel = WarningLevel.Warning;
+					result = testTargetRules.RequiresUniqueEnvironment(generatedTestRulesAssembly, null, [], out baseTarget);
+
+					Assert.IsFalse(result);
+
+					testTargetRules.CppCompileWarningSettings.DeprecationWarningLevel = preWarningLevel;
+				}
+
+				// Mutate the object against a RequiresUniqueEnvironment, and set bOverrideBuildEnvironment=true
+				{
+					WarningLevel preWarningLevel = testTargetRules.CppCompileWarningSettings.ShadowVariableWarningLevel;
+					testTargetRules.CppCompileWarningSettings.ShadowVariableWarningLevel = WarningLevel.Warning;
+					testTargetRules.bOverrideBuildEnvironment = true;
+
+					bool result = testTargetRules.RequiresUniqueEnvironment(generatedTestRulesAssembly, null, [], out baseTarget);
+
+					Assert.IsFalse(result);
+
+					testTargetRules.CppCompileWarningSettings.ShadowVariableWarningLevel = preWarningLevel;
+				}
+			}
+		}
+	}
+}
\ No newline at end of file
diff --git a/Engine/Source/Programs/UnrealBuildTool.Tests/TestUtilities/SafeTestDirectory.cs b/Engine/Source/Programs/UnrealBuildTool.Tests/TestUtilities/SafeTestDirectory.cs
new file mode 100644
--- /dev/null
+++ b/Engine/Source/Programs/UnrealBuildTool.Tests/TestUtilities/SafeTestDirectory.cs
@@ -0,0 +1,46 @@
+// Copyright Epic Games, Inc. All Rights Reserved.
+
+using System;
+using System.IO;
+
+namespace UnrealBuildTool.Tests.TestUtilities
+{
+	internal class SafeTestDirectory : IDisposable
+	{
+		internal string TemporaryDirectory { get; private set; }
+		private bool _wasFolderCreated = false;
+
+		private SafeTestDirectory(string temporaryDirectory, bool wasFolderCreated = false)
+		{
+			TemporaryDirectory = temporaryDirectory;
+			_wasFolderCreated = wasFolderCreated;
+		}
+
+		internal static SafeTestDirectory CreateTestDirectory(string absoluteDirectoryPath)
+		{
+			SafeTestDirectory returnDirectory = new(absoluteDirectoryPath);
+			if (!Directory.Exists(absoluteDirectoryPath))
+			{
+				Directory.CreateDirectory(absoluteDirectoryPath);
+				returnDirectory._wasFolderCreated = true;
+			}
+
+			return returnDirectory;
+		}
+
+		public void Dispose()
+		{
+			if (_wasFolderCreated && !String.IsNullOrEmpty(TemporaryDirectory) && Directory.Exists(TemporaryDirectory))
+			{
+				try
+				{
+					Directory.Delete(TemporaryDirectory, true);
+				}
+				catch (Exception)
+				{
+					// Swallow exception
+				}
+			}
+		}
+	}
+}
diff --git a/Engine/Source/Programs/UnrealBuildTool.Tests/TestUtilities/SafeTestFile.cs b/Engine/Source/Programs/UnrealBuildTool.Tests/TestUtilities/SafeTestFile.cs
new file mode 100644
--- /dev/null
+++ b/Engine/Source/Programs/UnrealBuildTool.Tests/TestUtilities/SafeTestFile.cs
@@ -0,0 +1,41 @@
+// Copyright Epic Games, Inc. All Rights Reserved.
+
+using System;
+using System.IO;
+using EpicGames.Core;
+
+namespace UnrealBuildTool.Tests.TestUtilities
+{
+	internal class SafeTestFile : IDisposable
+	{
+		internal string TemporaryFile { get; private set; }
+
+		internal SafeTestFile(string content, string outputFileName, string outputFolder)
+		{
+			try
+			{
+				TemporaryFile = Path.Combine(outputFolder, outputFileName);
+
+				using (StreamWriter writer = new(TemporaryFile))
+				{
+					writer.Write(content);
+				}
+			}
+			catch (Exception)
+			{
+				throw new InvalidOperationException("Unable to create test file. Test not feasible.");
+			}
+		}
+
+		public void Dispose()
+		{
+			if (!String.IsNullOrEmpty(TemporaryFile) && File.Exists(TemporaryFile))
+			{
+				File.Delete(TemporaryFile);
+			}
+		}
+
+		public static implicit operator FileReference(SafeTestFile file)
+			=> new(file.TemporaryFile);
+	}
+}
diff --git a/Engine/Source/Programs/UnrealBuildTool.Tests/TestUtilities/StaticTargetRulesInitializer.cs b/Engine/Source/Programs/UnrealBuildTool.Tests/TestUtilities/StaticTargetRulesInitializer.cs
new file mode 100644
--- /dev/null
+++ b/Engine/Source/Programs/UnrealBuildTool.Tests/TestUtilities/StaticTargetRulesInitializer.cs
@@ -0,0 +1,62 @@
+// Copyright Epic Games, Inc. All Rights Reserved.
+
+using System;
+using System.IO;
+using Microsoft.Extensions.Logging;
+using Microsoft.VisualStudio.TestTools.UnitTesting;
+
+namespace UnrealBuildTool.Tests.TestUtilities
+{
+	[TestClass]
+	public static class StaticTargetRulesInitializer
+	{
+		[AssemblyInitialize]
+		[System.Diagnostics.CodeAnalysis.SuppressMessage("Style", "IDE0060:Remove unused parameter", Justification = "Required for UnitTesting")]
+		public static void InitializeBuildEnvironment(TestContext context)
+		{
+			// Clean any previous runs linger temp folder
+			// Remarks: this is due to target rules modules dll being held by the test process; full cleanup not feasible on existing run
+			string[] subfolders = Directory.GetDirectories(TestDirs.UBTTestFolderRoot, $"{TestDirs.UBTTestFolderPrefix}*", SearchOption.TopDirectoryOnly);
+
+			foreach (string folder in subfolders)
+			{
+				if (Directory.Exists(folder))
+				{
+					try
+					{
+						Directory.Delete(folder, true);
+					}
+					catch (Exception) { }
+				}
+			}
+
+			// Example, on how to register a fully stubbed test environment.
+			/*
+			 	UEBuildPlatform.RegisterBuildPlatform(
+					new TestBuildPlatfrom(
+							UnrealTargetPlatform.Win64,
+							new TestSDK(logger),
+							new UnrealArchitectureConfig(UnrealArch.X64),
+							logger),
+					logger
+			);*/
+
+			ILogger logger = new TestLogger();
+
+			UEBuildPlatform? _;
+			if (!UEBuildPlatform.TryGetBuildPlatform(UnrealTargetPlatform.Win64, out _))
+			{
+				MicrosoftPlatformSDK sdk = new(logger);
+				UEBuildPlatform.RegisterBuildPlatform(new WindowsPlatform(UnrealTargetPlatform.Win64, sdk, logger), logger);
+			}
+
+			if (!UEBuildPlatform.TryGetBuildPlatform(UnrealTargetPlatform.Android, out _))
+			{
+				AndroidPlatformSDK asdk = new(logger);
+				UEBuildPlatform.RegisterBuildPlatform(new AndroidPlatform(asdk, logger), logger);
+			}
+
+			XmlConfig.ReadConfigFiles(null, null, logger);
+		}
+	}
+}
diff --git a/Engine/Source/Programs/UnrealBuildTool.Tests/TestUtilities/TestBuildPlatform.cs b/Engine/Source/Programs/UnrealBuildTool.Tests/TestUtilities/TestBuildPlatform.cs
new file mode 100644
--- /dev/null
+++ b/Engine/Source/Programs/UnrealBuildTool.Tests/TestUtilities/TestBuildPlatform.cs
@@ -0,0 +1,39 @@
+// Copyright Epic Games, Inc. All Rights Reserved.
+
+using EpicGames.Core;
+using Microsoft.Extensions.Logging;
+
+namespace UnrealBuildTool.Tests.TestUtilities
+{
+	internal class TestBuildPlatfrom : UEBuildPlatform
+	{
+		public TestBuildPlatfrom(UnrealTargetPlatform inPlatform, UEBuildPlatformSDK sdk, UnrealArchitectureConfig architectureConfig, ILogger inLogger) : base(inPlatform, sdk, architectureConfig, inLogger)
+		{
+
+		}
+		public override UEToolChain CreateToolChain(ReadOnlyTargetRules target)
+		{
+			throw new System.NotImplementedException();
+		}
+
+		public override void Deploy(TargetReceipt receipt)
+		{
+
+		}
+
+		public override bool IsBuildProduct(string fileName, string[] namePrefixes, string[] nameSuffixes)
+		{
+			return false;
+		}
+
+		public override void SetUpEnvironment(ReadOnlyTargetRules target, CppCompileEnvironment compileEnvironment, LinkEnvironment linkEnvironment)
+		{
+
+		}
+
+		public override bool ShouldCreateDebugInfo(ReadOnlyTargetRules target)
+		{
+			return false;
+		}
+	}
+}
diff --git a/Engine/Source/Programs/UnrealBuildTool.Tests/TestUtilities/TestBuildPlatformFactory.cs b/Engine/Source/Programs/UnrealBuildTool.Tests/TestUtilities/TestBuildPlatformFactory.cs
new file mode 100644
--- /dev/null
+++ b/Engine/Source/Programs/UnrealBuildTool.Tests/TestUtilities/TestBuildPlatformFactory.cs
@@ -0,0 +1,15 @@
+// Copyright Epic Games, Inc. All Rights Reserved.
+
+using Microsoft.Extensions.Logging;
+
+namespace UnrealBuildTool.Tests.TestUtilities
+{
+	internal class TestBuildPlatformFactory : UEBuildPlatformFactory
+	{
+		public override UnrealTargetPlatform TargetPlatform => UnrealTargetPlatform.Win64;
+
+		public override void RegisterBuildPlatforms(ILogger logger)
+		{
+		}
+	}
+}
diff --git a/Engine/Source/Programs/UnrealBuildTool.Tests/TestUtilities/TestBuildPlatformSDK.cs b/Engine/Source/Programs/UnrealBuildTool.Tests/TestUtilities/TestBuildPlatformSDK.cs
new file mode 100644
--- /dev/null
+++ b/Engine/Source/Programs/UnrealBuildTool.Tests/TestUtilities/TestBuildPlatformSDK.cs
@@ -0,0 +1,25 @@
+// Copyright Epic Games, Inc. All Rights Reserved.
+
+using EpicGames.Core;
+using Microsoft.Extensions.Logging;
+
+namespace UnrealBuildTool.Tests.TestUtilities
+{
+	internal class TestSDK : UEBuildPlatformSDK
+	{
+		public TestSDK(ILogger inLogger) : base(inLogger)
+		{
+		}
+
+		public override bool TryConvertVersionToInt(string? stringValue, out ulong outValue, string? hint = null)
+		{
+			outValue = 1;
+			return true;
+		}
+
+		protected override string? GetInstalledSDKVersion()
+		{
+			return "1.0.0";
+		}
+	}
+}
diff --git a/Engine/Source/Programs/UnrealBuildTool.Tests/TestUtilities/TestDirs.cs b/Engine/Source/Programs/UnrealBuildTool.Tests/TestUtilities/TestDirs.cs
new file mode 100644
--- /dev/null
+++ b/Engine/Source/Programs/UnrealBuildTool.Tests/TestUtilities/TestDirs.cs
@@ -0,0 +1,31 @@
+// Copyright Epic Games, Inc. All Rights Reserved.
+
+using System;
+using System.IO;
+using EpicGames.Core;
+using UnrealBuildBase;
+
+namespace UnrealBuildTool.Tests.TestUtilities
+{
+	internal class TestDirs
+	{
+		internal static DirectoryReference Stubs
+		{
+			get
+			{
+				DirectoryReference programsDirectory = DirectoryReference.Combine(Unreal.EngineSourceDirectory, "Programs");
+				return DirectoryReference.Combine(programsDirectory, "UnrealBuildTool.Tests", "UBT");
+			}
+		}
+
+		internal const string UBTTestFolderPrefix = "UBT-Test-";
+		internal static readonly string UBTTestFolderRoot = Path.GetTempPath();
+
+		internal static SafeTestDirectory Create()
+		{
+			string tempDirName = $"{UBTTestFolderPrefix}{Guid.NewGuid()}";
+			string tempDirPath = Path.Combine(UBTTestFolderRoot, tempDirName);
+			return SafeTestDirectory.CreateTestDirectory(tempDirPath);
+		}
+	}
+}
diff --git a/Engine/Source/Programs/UnrealBuildTool.Tests/TestUtilities/TestLogger.cs b/Engine/Source/Programs/UnrealBuildTool.Tests/TestUtilities/TestLogger.cs
new file mode 100644
--- /dev/null
+++ b/Engine/Source/Programs/UnrealBuildTool.Tests/TestUtilities/TestLogger.cs
@@ -0,0 +1,24 @@
+// Copyright Epic Games, Inc. All Rights Reserved.
+
+using System;
+using Microsoft.Extensions.Logging;
+
+namespace UnrealBuildTool.Tests.TestUtilities
+{
+	internal class TestLogger : ILogger
+	{
+		public IDisposable? BeginScope<TState>(TState state) where TState : notnull
+		{
+			return null;
+		}
+
+		public bool IsEnabled(LogLevel logLevel)
+		{
+			return false;
+		}
+
+		public void Log<TState>(LogLevel logLevel, EventId eventId, TState state, Exception? exception, Func<TState, Exception?, string> formatter)
+		{
+		}
+	}
+}
\ No newline at end of file
diff --git a/Engine/Source/Programs/UnrealBuildTool.Tests/TestUtilities/TestToolChain.cs b/Engine/Source/Programs/UnrealBuildTool.Tests/TestUtilities/TestToolChain.cs
new file mode 100644
--- /dev/null
+++ b/Engine/Source/Programs/UnrealBuildTool.Tests/TestUtilities/TestToolChain.cs
@@ -0,0 +1,36 @@
+// Copyright Epic Games, Inc. All Rights Reserved.
+
+using System.Collections.Generic;
+using EpicGames.Core;
+using Microsoft.Extensions.Logging;
+using UnrealBuildBase;
+
+namespace UnrealBuildTool.Tests.TestUtilities
+{
+	internal class TestToolChain : UEToolChain
+	{
+		public TestToolChain(ILogger inLogger) : base(inLogger)
+		{
+		}
+
+		public override FileItem? LinkFiles(LinkEnvironment linkEnvironment, bool bBuildImportLibraryOnly, IActionGraphBuilder graph)
+		{
+			return null;
+		}
+
+		protected override CPPOutput CompileCPPFiles(CppCompileEnvironment compileEnvironment, IEnumerable<FileItem> inputFiles, DirectoryReference outputDir, string moduleName, IActionGraphBuilder graph)
+		{
+			return null!;
+		}
+
+		protected override CPPOutput CompileISPCFiles(CppCompileEnvironment compileEnvironment, IEnumerable<FileItem> inputFiles, DirectoryReference outputDir, IActionGraphBuilder graph)
+		{
+			return null!;
+		}
+
+		protected override CPPOutput GenerateISPCHeaders(CppCompileEnvironment compileEnvironment, IEnumerable<FileItem> inputFiles, DirectoryReference outputDir, IActionGraphBuilder graph)
+		{
+			return null!;
+		}
+	}
+}
diff --git a/Engine/Source/Programs/UnrealBuildTool.Tests/TokenReaderTests.cs b/Engine/Source/Programs/UnrealBuildTool.Tests/TokenReaderTests.cs
new file mode 100644
--- /dev/null
+++ b/Engine/Source/Programs/UnrealBuildTool.Tests/TokenReaderTests.cs
@@ -0,0 +1,172 @@
+// Copyright Epic Games, Inc. All Rights Reserved.
+
+using System;
+using System.Collections.Generic;
+using System.Linq;
+using System.Text;
+using Microsoft.VisualStudio.TestTools.UnitTesting;
+
+namespace UnrealBuildTool.Tests
+{
+	[TestClass]
+	public class TokenReaderTests
+	{
+		[TestMethod]
+		public void Run()
+		{
+			RunLexerTest("ABC//hello\nDEF", "Identifier(ABC), Newline, Identifier(DEF)");
+			RunLexerTest("ABC/*hello\nworld*/DEF", "Identifier(ABC), Identifier(DEF)");
+			RunLexerTest("ABC//hello\r\nDEF", "Identifier(ABC), Newline, Identifier(DEF)");
+			RunLexerTest("ABC123", "Identifier(ABC123)");
+			RunLexerTest("123ABC", "Number(123ABC)");
+			RunLexerTest("FOO(1,2,3,...)", "Identifier(FOO), LeftParen, Number(1), Comma, Number(2), Comma, Number(3), Comma, Ellipsis, RightParen");
+		}
+
+		[TestMethod]
+		public void SimpleTokens()
+		{
+			RunLexerTest("=<><<+", "Equals, CompareLess, CompareGreater, LeftShift, Plus");
+			RunLexerTest("+=+++++", "Plus, Equals, Plus, Plus, Plus, Plus, Plus");
+			RunLexerTest("/+++=", "Divide, Plus, Plus, Plus, Equals");
+			RunLexerTest("###", "HashHash, Hash");
+			RunLexerTest("():?", "LeftParen, RightParen, Colon, QuestionMark");
+			RunLexerTest("..*->->*", "Dot, Dot, Multiply, Minus, CompareGreater, Minus, CompareGreater, Multiply");
+			RunLexerTest("+-*/%", "Plus, Minus, Multiply, Divide, Modulo");
+			RunLexerTest("|^&~", "BitwiseOr, BitwiseXor, BitwiseAnd, BitwiseNot");
+			RunLexerTest("+=-=*=/=", "Plus, Equals, Minus, Equals, Multiply, Equals, Divide, Equals");
+			RunLexerTest("%=|=^=", "Modulo, Equals, BitwiseOr, Equals, BitwiseXor, Equals");
+			RunLexerTest("&&||!", "LogicalAnd, LogicalOr, LogicalNot");
+			RunLexerTest("=<><=>=", "Equals, CompareLess, CompareGreater, CompareLessOrEqual, CompareGreaterOrEqual");
+			RunLexerTest("==!=", "CompareEqual, CompareNotEqual");
+			RunLexerTest(".....", "Ellipsis, Dot, Dot");
+		}
+
+		[TestMethod]
+		public void IntegerLiterals()
+		{
+			RunLexerTest("123+", "Number(123), Plus");
+			RunLexerTest("01234567+", "Number(01234567), Plus");
+			RunLexerTest("0x123456789ABCDEFabcdef+", "Number(0x123456789ABCDEFabcdef), Plus");
+			RunLexerTest("123Ull+", "Number(123Ull), Plus");
+			RunLexerTest("123_hello+", "Number(123_hello), Plus");
+			RunLexerTest("01234567_hello+", "Number(01234567_hello), Plus");
+			RunLexerTest("0x123456789ABCDEFabcdef_hello+", "Number(0x123456789ABCDEFabcdef_hello), Plus");
+		}
+
+		[TestMethod]
+		public void FloatLiterals()
+		{
+			RunLexerTest("123.+", "Number(123.), Plus");
+			RunLexerTest("123e123+", "Number(123e123), Plus");
+			RunLexerTest("123e+123+", "Number(123e+123), Plus");
+			RunLexerTest(".123+", "Number(.123), Plus");
+			RunLexerTest(".123e123+", "Number(.123e123), Plus");
+			RunLexerTest(".123e+123+", "Number(.123e+123), Plus");
+			RunLexerTest("0.123+", "Number(0.123), Plus");
+			RunLexerTest("0.123e123+", "Number(0.123e123), Plus");
+			RunLexerTest("0.123e+123+", "Number(0.123e+123), Plus");
+			RunLexerTest("123.123+", "Number(123.123), Plus");
+			RunLexerTest("123.123e123+", "Number(123.123e123), Plus");
+			RunLexerTest("123.123e+123+", "Number(123.123e+123), Plus");
+			RunLexerTest("1.f+", "Number(1.f), Plus");
+			RunLexerTest("123.123e+123f+", "Number(123.123e+123f), Plus");
+			RunLexerTest("123.123e+123F+", "Number(123.123e+123F), Plus");
+		}
+
+		[TestMethod]
+		public void CharacterLiterals()
+		{
+			RunLexerTest("'hello'+", "Character('hello'), Plus");
+			RunLexerTest("'hello\\'\\''+", "Character('hello\\'\\''), Plus");
+			RunLexerTest("u'hello'+", "Character(u'hello'), Plus");
+			RunLexerTest("U'hello'+", "Character(U'hello'), Plus");
+			RunLexerTest("L'hello'+", "Character(L'hello'), Plus");
+			RunLexerTest("u8'hello'+", "Character(u8'hello'), Plus");
+		}
+
+		[TestMethod]
+		public void StringLiterals()
+		{
+			RunLexerTest("\"hello\"+", "String(\"hello\"), Plus");
+			RunLexerTest("\"hello\\\"\\\"\"+", "String(\"hello\\\"\\\"\"), Plus");
+			RunLexerTest("u\"hello\"+", "String(u\"hello\"), Plus");
+			RunLexerTest("U\"hello\"+", "String(U\"hello\"), Plus");
+			RunLexerTest("L\"hello\"+", "String(L\"hello\"), Plus");
+			RunLexerTest("u8\"hello\"+", "String(u8\"hello\"), Plus");
+		}
+
+		[TestMethod]
+		public void Identifiers()
+		{
+			RunLexerTest("abcdefghijklmnopqrstuvwxyz1234567890ABCDEFGHIJKLMNOPQRSTUVWXYZ_+", "Identifier(abcdefghijklmnopqrstuvwxyz1234567890ABCDEFGHIJKLMNOPQRSTUVWXYZ_), Plus");
+		}
+
+		[TestMethod]
+		public void Comments()
+		{
+			// Single-line comments
+			RunLexerTest("//hello\n+", "Newline, Plus");
+			RunLexerTest("//hello\n//world\n+", "Newline, Newline, Plus");
+
+			// Multi-line comments
+			RunLexerTest("/*hello*/+", "Plus");
+			RunLexerTest("/*/hello*/+", "Plus");
+			RunLexerTest("/*/hello**/+", "Plus");
+		}
+
+		[TestMethod]
+		public void LineNumbers()
+		{
+			// Test reported line numbers are accurate
+			RunLineNumberTest("+", 1);
+			RunLineNumberTest("\n+", 2);
+			RunLineNumberTest("\r\n+", 2);
+			RunLineNumberTest("\r\n\r\n+", 3);
+			RunLineNumberTest("\r\n\n+", 3);
+			RunLineNumberTest("//comment\n+", 2);
+			RunLineNumberTest("//comment\\\n\n+", 3);
+			RunLineNumberTest("/*comment\n*/+", 2);
+			RunLineNumberTest("/*comment\r\n*/+", 2);
+			RunLineNumberTest("/*comment\r\n\r\n*/+", 3);
+			RunLineNumberTest("first line \\\n+", 2);
+		}
+
+		public static void RunLexerTest(string inputText, string expectedOutputText)
+		{
+			List<Token> tokens = [];
+
+			using TokenReader reader = new(inputText);
+			while (reader.MoveNext())
+			{
+				tokens.Add(reader.Current);
+			}
+
+			string outputText = String.Join(", ", tokens.Select(FormatToken));
+			Assert.AreEqual(expectedOutputText, outputText);
+		}
+
+		static string FormatToken(Token token)
+		{
+			StringBuilder result = new();
+			result.AppendFormat(token.Type.ToString());
+			if (token.Type == TokenType.Identifier || token.Type == TokenType.Number || token.Type == TokenType.Character || token.Type == TokenType.String)
+			{
+				result.AppendFormat("({0})", token.Text);
+			}
+			return result.ToString();
+		}
+
+		public static void RunLineNumberTest(string inputText, int expectedLineNumber)
+		{
+			using TokenReader reader = new(inputText);
+			while (reader.MoveNext())
+			{
+				if (reader.Current.Type == TokenType.Plus)
+				{
+					break;
+				}
+			}
+			Assert.AreEqual(expectedLineNumber, reader.LineNumber);
+		}
+	}
+}
diff --git a/Engine/Source/Programs/UnrealBuildTool.Tests/UnrealBuildTool.Tests.csproj b/Engine/Source/Programs/UnrealBuildTool.Tests/UnrealBuildTool.Tests.csproj
new file mode 100644
--- /dev/null
+++ b/Engine/Source/Programs/UnrealBuildTool.Tests/UnrealBuildTool.Tests.csproj
@@ -0,0 +1,42 @@
+<Project Sdk="Microsoft.NET.Sdk">
+  <Import Project="..\Shared\UnrealEngine.csproj.props" />
+
+  <PropertyGroup>
+    <Nullable>enable</Nullable>
+    <OutputType>Library</OutputType>
+    <GenerateAssemblyInfo>false</GenerateAssemblyInfo>
+    <IsPackable>false</IsPackable>
+    <IsTestProject>true</IsTestProject>
+    <IsPackable>false</IsPackable>
+    <IsTestProject>true</IsTestProject>
+    <!-- We don't require public documentation for test projects -->
+    <NoWarn>$(NoWarn);1591</NoWarn>
+  </PropertyGroup>
+
+  <ItemGroup>
+    <Compile Include="..\Shared\MetaData.cs">
+      <Link>Properties\MetaData.cs</Link>
+    </Compile>
+  </ItemGroup>
+
+  <ItemGroup>
+    <PackageReference Include="Microsoft.NET.Test.Sdk" Version="17.14.1">
+      <TreatAsUsed>true</TreatAsUsed>
+    </PackageReference>
+    <PackageReference Include="MSTest.TestAdapter" Version="3.11.1">
+      <TreatAsUsed>true</TreatAsUsed>
+    </PackageReference>
+    <PackageReference Include="MSTest.TestFramework" Version="3.11.1">
+      <TreatAsUsed>true</TreatAsUsed>
+    </PackageReference>
+    <PackageReference Include="coverlet.collector" Version="6.0.4">
+      <PrivateAssets>all</PrivateAssets>
+      <IncludeAssets>runtime; build; native; contentfiles; analyzers; buildtransitive</IncludeAssets>
+    </PackageReference>
+  </ItemGroup>
+
+  <ItemGroup>
+    <ProjectReference Include="..\Shared\EpicGames.Horde\EpicGames.Horde.csproj" />
+    <ProjectReference Include="..\UnrealBuildTool\UnrealBuildTool.csproj" />
+  </ItemGroup>
+</Project>
diff --git a/Engine/Source/Programs/UnrealBuildTool.Tests/WarningLevelDefaultAttributeTests.cs b/Engine/Source/Programs/UnrealBuildTool.Tests/WarningLevelDefaultAttributeTests.cs
new file mode 100644
--- /dev/null
+++ b/Engine/Source/Programs/UnrealBuildTool.Tests/WarningLevelDefaultAttributeTests.cs
@@ -0,0 +1,136 @@
+// Copyright Epic Games, Inc. All Rights Reserved.
+
+using System.Collections.Generic;
+using Microsoft.VisualStudio.TestTools.UnitTesting;
+using UnrealBuildTool.Configuration.CompileWarnings;
+
+namespace UnrealBuildTool.Tests
+{
+	[TestClass]
+	public class WarningLevelDefaultAttributeTests
+	{
+		internal class TestBuildSettingsProvider : IBuildContextProvider
+		{
+			internal BuildSettingsVersion _testBuildSettings = BuildSettingsVersion.V1;
+
+			public BuildSettingsVersion GetBuildSettings()
+			{
+				return _testBuildSettings;
+			}
+		}
+
+		[TestMethod]
+		public void TestWindowsCompletelyDefined()
+		{
+			List<VersionWarningLevelDefaultAttribute> attributes =
+			[
+				new VersionWarningLevelDefaultAttribute( WarningLevel.Off, BuildSettingsVersion.V1, BuildSettingsVersion.V2),
+				new VersionWarningLevelDefaultAttribute(WarningLevel.Warning, BuildSettingsVersion.V3, BuildSettingsVersion.V4),
+				new VersionWarningLevelDefaultAttribute(WarningLevel.Error, BuildSettingsVersion.V5, BuildSettingsVersion.Latest)
+			];
+
+			List<VersionWarningLevelDefaultAttribute> merged = VersionWarningLevelDefaultAttribute.EnsureWarningLevelDefaultBounds(attributes);
+
+			Assert.AreEqual(3, merged.Count);
+			Assert.AreEqual(BuildSettingsVersion.V1, merged[0].MinVersion);
+			Assert.AreEqual(BuildSettingsVersion.V2, merged[0].MaxVersion);
+			Assert.AreEqual(WarningLevel.Off, merged[0].GetDefaultLevel());
+
+			Assert.AreEqual(BuildSettingsVersion.V3, merged[1].MinVersion);
+			Assert.AreEqual(BuildSettingsVersion.V4, merged[1].MaxVersion);
+			Assert.AreEqual(WarningLevel.Warning, merged[1].GetDefaultLevel());
+
+			Assert.AreEqual(BuildSettingsVersion.V5, merged[2].MinVersion);
+			Assert.AreEqual(BuildSettingsVersion.Latest, merged[2].MaxVersion);
+			Assert.AreEqual(WarningLevel.Error, merged[2].GetDefaultLevel());
+		}
+
+		[TestMethod]
+		public void TestGapResolveToLowerBound()
+		{
+			List<VersionWarningLevelDefaultAttribute> attributes =
+			[
+				new VersionWarningLevelDefaultAttribute(WarningLevel.Warning, BuildSettingsVersion.V1, BuildSettingsVersion.V2),
+				new VersionWarningLevelDefaultAttribute(WarningLevel.Error, BuildSettingsVersion.V4, BuildSettingsVersion.V5),
+			];
+
+			List<VersionWarningLevelDefaultAttribute> merged = VersionWarningLevelDefaultAttribute.EnsureWarningLevelDefaultBounds(attributes);
+
+			Assert.AreEqual(2, merged.Count);
+			Assert.AreEqual(BuildSettingsVersion.V1, merged[0].MinVersion);
+			Assert.AreEqual(BuildSettingsVersion.V3, merged[0].MaxVersion);
+			Assert.AreEqual(WarningLevel.Warning, merged[0].GetDefaultLevel());
+
+			Assert.AreEqual(BuildSettingsVersion.V4, merged[1].MinVersion);
+			Assert.AreEqual(BuildSettingsVersion.V5, merged[1].MaxVersion);
+			Assert.AreEqual(WarningLevel.Error, merged[1].GetDefaultLevel());
+		}
+
+		[TestMethod]
+		public void TestLowerUpperOverlapResizeBottomUpward()
+		{
+			List<VersionWarningLevelDefaultAttribute> attributes =
+			[
+				new VersionWarningLevelDefaultAttribute(WarningLevel.Warning, BuildSettingsVersion.V1, BuildSettingsVersion.V3),
+				new VersionWarningLevelDefaultAttribute(WarningLevel.Error, BuildSettingsVersion.V2, BuildSettingsVersion.V4)
+			];
+
+			List<VersionWarningLevelDefaultAttribute> merged = VersionWarningLevelDefaultAttribute.EnsureWarningLevelDefaultBounds(attributes);
+
+			Assert.AreEqual(2, merged.Count);
+			Assert.AreEqual(BuildSettingsVersion.V1, merged[0].MinVersion);
+			Assert.AreEqual(BuildSettingsVersion.V3, merged[0].MaxVersion);
+			Assert.AreEqual(WarningLevel.Warning, merged[0].GetDefaultLevel());
+
+			Assert.AreEqual(BuildSettingsVersion.V4, merged[1].MinVersion);
+			Assert.AreEqual(BuildSettingsVersion.V4, merged[1].MaxVersion);
+			Assert.AreEqual(WarningLevel.Error, merged[1].GetDefaultLevel());
+		}
+
+		[TestMethod]
+		public void TestAllRangeWide()
+		{
+			List<VersionWarningLevelDefaultAttribute> attributes =
+			[
+				new VersionWarningLevelDefaultAttribute( WarningLevel.Warning, BuildSettingsVersion.V1, BuildSettingsVersion.V2),
+				new VersionWarningLevelDefaultAttribute(WarningLevel.Error, BuildSettingsVersion.V1, BuildSettingsVersion.V4)
+			];
+
+			List<VersionWarningLevelDefaultAttribute> merged = VersionWarningLevelDefaultAttribute.EnsureWarningLevelDefaultBounds(attributes);
+
+			Assert.AreEqual(2, merged.Count);
+
+			Assert.AreEqual(BuildSettingsVersion.V1, merged[0].MinVersion);
+			Assert.AreEqual(BuildSettingsVersion.V2, merged[0].MaxVersion);
+			Assert.AreEqual(WarningLevel.Warning, merged[0].GetDefaultLevel());
+
+			Assert.AreEqual(BuildSettingsVersion.V3, merged[1].MinVersion);
+			Assert.AreEqual(BuildSettingsVersion.V4, merged[1].MaxVersion);
+			Assert.AreEqual(WarningLevel.Error, merged[1].GetDefaultLevel());
+		}
+
+		[TestMethod]
+		public void TestResolveLowerBound()
+		{
+			List<WarningLevelDefaultAttribute> attributes =
+			[
+				new VersionWarningLevelDefaultAttribute(WarningLevel.Warning, BuildSettingsVersion.V1, BuildSettingsVersion.V2),
+				new VersionWarningLevelDefaultAttribute(WarningLevel.Error, BuildSettingsVersion.V3, BuildSettingsVersion.V4)
+			];
+			TestBuildSettingsProvider testBuildSettingsProvider = new();
+			BuildSystemContext newBuildSystemContext = new(testBuildSettingsProvider);
+
+			Assert.AreEqual(WarningLevel.Warning, VersionWarningLevelDefaultAttribute.ResolveVersionWarningLevelDefault(attributes, newBuildSystemContext));
+			testBuildSettingsProvider._testBuildSettings = BuildSettingsVersion.V2;
+			Assert.AreEqual(WarningLevel.Warning, VersionWarningLevelDefaultAttribute.ResolveVersionWarningLevelDefault(attributes, newBuildSystemContext));
+
+			testBuildSettingsProvider._testBuildSettings = BuildSettingsVersion.V3;
+			Assert.AreEqual(WarningLevel.Error, VersionWarningLevelDefaultAttribute.ResolveVersionWarningLevelDefault(attributes, newBuildSystemContext));
+			testBuildSettingsProvider._testBuildSettings = BuildSettingsVersion.V4;
+			Assert.AreEqual(WarningLevel.Error, VersionWarningLevelDefaultAttribute.ResolveVersionWarningLevelDefault(attributes, newBuildSystemContext));
+
+			testBuildSettingsProvider._testBuildSettings = BuildSettingsVersion.V5;
+			Assert.AreEqual(WarningLevel.Default, VersionWarningLevelDefaultAttribute.ResolveVersionWarningLevelDefault(attributes, newBuildSystemContext));
+		}
+	}
+}
\ No newline at end of file
