Adding condensed graph to c#.
This is an implementation of Kosaraju-Sharir's in C# a port of the existing C++ GraphConvert.cpp

Tests confirmed with:
dotnet test ".\Engine\Source\Programs\UnrealBuildTool.Tests\UnrealBuildTool.Tests.csproj" --filter "ClassName=UnrealBuildTool.Tests.CondensationGraphBuilderTests"

#Synced-CL 52439769

diff --git a/Engine/Source/Programs/UnrealBuildTool.Tests/CondensationGraphBuilderTests.cs b/Engine/Source/Programs/UnrealBuildTool.Tests/CondensationGraphBuilderTests.cs
new file mode 100644
--- /dev/null
+++ b/Engine/Source/Programs/UnrealBuildTool.Tests/CondensationGraphBuilderTests.cs
@@ -0,0 +1,508 @@
+// Copyright Epic Games, Inc. All Rights Reserved.
+
+using System;
+using System.Collections.Generic;
+using System.Linq;
+using Microsoft.VisualStudio.TestTools.UnitTesting;
+
+namespace UnrealBuildTool.Tests
+{
+	/// <summary>
+	/// Tests for CondensationGraphBuilder, ported from the C++ automation test
+	/// FCondensationGraphTest in Editor/UnrealEd/Private/Cooker/Algo/GraphConvert.cpp.
+	/// </summary>
+	[TestClass]
+	public class CondensationGraphBuilderTests
+	{
+		// -------------------------------------------------------------------------
+		// Test state (reset between cases via Clear())
+		// -------------------------------------------------------------------------
+
+		/// <summary>
+		/// Sparse adjacency list: _adjacency[v] contains the outgoing edge targets for vertex v.
+		/// </summary>
+		private readonly List<int[]> _adjacency = [];
+
+		/// <summary>
+		/// Expected SCC membership for vertices that participate in cycles.
+		/// Vertices not listed here are expected to be singleton SCCs.
+		/// </summary>
+		private readonly List<int[]> _expectedComponents = [];
+
+		// -------------------------------------------------------------------------
+		// Helpers mirroring the C++ AddVertex / AddExpectedComponent / Clear lambdas
+		// -------------------------------------------------------------------------
+
+		private void AddVertex(int v, params int[] edges)
+		{
+			while (_adjacency.Count <= v)
+			{
+				_adjacency.Add([]);
+			}
+			_adjacency[v] = edges;
+		}
+
+		private void AddExpectedComponent(params int[] vertices)
+		{
+			_expectedComponents.Add(vertices);
+		}
+
+		private void Clear()
+		{
+			_adjacency.Clear();
+			_expectedComponents.Clear();
+		}
+
+		// -------------------------------------------------------------------------
+		// ConfirmResults — mirrors the C++ ConfirmResults lambda
+		// -------------------------------------------------------------------------
+
+		private void ConfirmResults(string testCaseName)
+		{
+			// Build the edge list from _adjacency.
+			List<Tuple<int, int>> edges = [];
+			for (int v = 0; v < _adjacency.Count; ++v)
+			{
+				foreach (int u in _adjacency[v])
+				{
+					edges.Add(Tuple.Create(v, u));
+				}
+			}
+
+			CondensationGraphBuilder<int> builder = new(edges);
+			bool result = builder.Build();
+
+			// (1) Verify the return value matches whether we declared any expected components.
+			bool expectedResult = _expectedComponents.Count > 0;
+			Assert.AreEqual(expectedResult, result,
+				$"[{testCaseName}] Build() returned {result} but expected {expectedResult}.");
+
+			List<int[]>                  sCCs            = builder.GetSCCs();
+			IReadOnlyDictionary<int, int> nodeToSCC      = builder.GetNodeToSCCIndex();
+			List<int[]>                  condensedEdges  = builder.GetCondensedEdges();
+
+			int numVertices   = _adjacency.Count;
+			int numOutVertices = sCCs.Count;
+
+			// (2) Verify that OutVertexToInVertices and InVertexToOutVertex are consistent:
+			//     every vertex appears in exactly one SCC and the two mappings agree.
+			int[] actualInVertexToOutVertex = new int[numVertices];
+			bool[] assigned = new bool[numVertices];
+
+			int componentIndex = 0;
+			foreach (int[] componentVertices in sCCs)
+			{
+				foreach (int inVertex in componentVertices)
+				{
+					Assert.IsFalse(assigned[inVertex],
+						$"[{testCaseName}] OutVertexToInVertices has vertex {inVertex} in multiple SCCs " +
+						$"(seen again in SCC {componentIndex}, previously in SCC {actualInVertexToOutVertex[inVertex]}).");
+
+					assigned[inVertex] = true;
+					actualInVertexToOutVertex[inVertex] = componentIndex;
+				}
+				++componentIndex;
+			}
+
+			for (int inVertex = 0; inVertex < numVertices; ++inVertex)
+			{
+				Assert.IsTrue(assigned[inVertex],
+					$"[{testCaseName}] OutVertexToInVertices is missing vertex {inVertex}.");
+
+				Assert.AreEqual(actualInVertexToOutVertex[inVertex], nodeToSCC[inVertex],
+					$"[{testCaseName}] GetNodeToSCCIndex()[{inVertex}] = {nodeToSCC[inVertex]} " +
+					$"but GetSCCs() places it in SCC {actualInVertexToOutVertex[inVertex]}.");
+			}
+
+			// (3) Finish building ExpectedOutVertexToInVertices: any vertex not explicitly
+			//     declared is assumed to be a singleton SCC. Then verify actual memberships.
+			int[] expectedInVertexToOutVertex = new int[numVertices];
+			List<List<int>> allExpectedComponents = [];
+
+			Array.Clear(assigned, 0, numVertices);
+			int expComponentIndex = 0;
+			foreach (int[] comp in _expectedComponents)
+			{
+				List<int> compList = [.. comp];
+				allExpectedComponents.Add(compList);
+				foreach (int v in comp)
+				{
+					assigned[v] = true;
+					expectedInVertexToOutVertex[v] = expComponentIndex;
+				}
+				++expComponentIndex;
+			}
+			for (int inVertex = 0; inVertex < numVertices; ++inVertex)
+			{
+				if (!assigned[inVertex])
+				{
+					expectedInVertexToOutVertex[inVertex] = expComponentIndex++;
+					allExpectedComponents.Add([inVertex]);
+				}
+			}
+
+			// For each vertex, check its SCC has the same sorted members as expected.
+			Array.Clear(assigned, 0, numVertices);
+			for (int inVertex = 0; inVertex < numVertices; ++inVertex)
+			{
+				if (assigned[inVertex])
+				{
+					continue;
+				}
+
+				int actualSCCIndex   = actualInVertexToOutVertex[inVertex];
+				int expectedSCCIndex = expectedInVertexToOutVertex[inVertex];
+
+				int[] actualMembers   = [.. sCCs[actualSCCIndex].OrderBy(x => x)];
+				int[] expectedMembers = [.. allExpectedComponents[expectedSCCIndex].OrderBy(x => x)];
+
+				CollectionAssert.AreEqual(expectedMembers, actualMembers,
+					$"[{testCaseName}] Vertex {inVertex} is in SCC {{{String.Join(",", actualMembers)}}} " +
+					$"but expected SCC {{{String.Join(",", expectedMembers)}}}.");
+
+				foreach (int m in actualMembers)
+				{
+					assigned[m] = true;
+				}
+			}
+
+			if (result)
+			{
+				// (4) Every input edge must be either within-SCC or present in the condensation.
+				for (int inSourceVertex = 0; inSourceVertex < numVertices; ++inSourceVertex)
+				{
+					int outSourceVertex = actualInVertexToOutVertex[inSourceVertex];
+					foreach (int inEdgeVertex in _adjacency[inSourceVertex])
+					{
+						int expectedOutEdgeVertex = actualInVertexToOutVertex[inEdgeVertex];
+						if (expectedOutEdgeVertex == outSourceVertex)
+						{
+							continue;
+						}
+						Assert.IsTrue(condensedEdges[outSourceVertex].Contains(expectedOutEdgeVertex),
+							$"[{testCaseName}] Missing condensation edge: expected SCC {outSourceVertex} -> " +
+							$"SCC {expectedOutEdgeVertex} (from input edge {inSourceVertex} -> {inEdgeVertex}).");
+					}
+				}
+
+				// (5) No extra condensation edges: for each SCC, compute the expected edge set
+				//     by unioning all member outgoing edges and verify exact match.
+				for (int outVertex = 0; outVertex < numOutVertices; ++outVertex)
+				{
+					HashSet<int> expectedOutEdges = [];
+					foreach (int inVertex in sCCs[outVertex])
+					{
+						foreach (int sourceEdgeVertex in _adjacency[inVertex])
+						{
+							int outEdgeVertex = actualInVertexToOutVertex[sourceEdgeVertex];
+							if (outEdgeVertex != outVertex)
+							{
+								expectedOutEdges.Add(outEdgeVertex);
+							}
+						}
+					}
+
+					int[] sortedActual   = [.. condensedEdges[outVertex].OrderBy(x => x)];
+					int[] sortedExpected = [.. expectedOutEdges.OrderBy(x => x)];
+
+					CollectionAssert.AreEqual(sortedExpected, sortedActual,
+						$"[{testCaseName}] SCC {outVertex}: expected edges " +
+						$"{{{String.Join(",", sortedExpected)}}} but got " +
+						$"{{{String.Join(",", sortedActual)}}}.");
+				}
+
+				// (6) Verify condensation is topologically sorted root-to-leaf:
+				//     every vertex reachable from SCC[Root] must have index >= Root.
+				for (int root = 0; root < numOutVertices; ++root)
+				{
+					HashSet<int> seen = [root];
+					Stack<int> toVisit = new();
+					toVisit.Push(root);
+
+					while (toVisit.Count > 0)
+					{
+						int cur = toVisit.Pop();
+						foreach (int neighbor in condensedEdges[cur])
+						{
+							if (seen.Add(neighbor))
+							{
+								toVisit.Push(neighbor);
+							}
+						}
+					}
+
+					foreach (int reachableVertex in seen)
+					{
+						Assert.IsTrue(reachableVertex >= root,
+							$"[{testCaseName}] SCC {reachableVertex} is reachable from SCC {root} " +
+							$"but has a lower index (output is not root-to-leaf sorted).");
+					}
+				}
+			}
+			else
+			{
+				// (7) When Build() returns false, GetCondensedEdges() must be empty.
+				Assert.AreEqual(0, condensedEdges.Count,
+					$"[{testCaseName}] Expected GetCondensedEdges() to be empty when Build() returns false.");
+			}
+		}
+
+		// -------------------------------------------------------------------------
+		// Test cases — direct ports of the C++ FCondensationGraphTest cases
+		// -------------------------------------------------------------------------
+
+		[TestMethod]
+		public void EachNodeDependsOnThePreviousOne()
+		{
+			Clear();
+			AddVertex(0);
+			AddVertex(1, 0);
+			AddVertex(2, 1);
+			ConfirmResults("Each node depends on the previous one");
+		}
+
+		[TestMethod]
+		public void EachNodeDependsOnTheNextOne()
+		{
+			Clear();
+			AddVertex(0, 1);
+			AddVertex(1, 2);
+			AddVertex(2);
+			ConfirmResults("Each node depends on the next one");
+		}
+
+		[TestMethod]
+		public void SelfReferences()
+		{
+			Clear();
+			AddVertex(0, 0);
+			AddVertex(1, 0, 1);
+			AddVertex(2, 1, 2);
+			ConfirmResults("SelfReferences");
+		}
+
+		[TestMethod]
+		public void SketchedOutExample1()
+		{
+			//              6
+			//             / \
+			//            5   7
+			//           / \   \
+			//          0   1   8
+			//           \ / \   \
+			//            3   4   |
+			//             \   \ /
+			//              \   9
+			//               \ /
+			//                2
+			Clear();
+			AddVertex(6, 5, 7);
+			AddVertex(5, 0, 1);
+			AddVertex(7, 8);
+			AddVertex(0, 3);
+			AddVertex(1, 3, 4);
+			AddVertex(8, 9);
+			AddVertex(3, 2);
+			AddVertex(4, 9);
+			AddVertex(9, 2);
+			AddVertex(2);
+			ConfirmResults("SketchedOutExample1");
+		}
+
+		[TestMethod]
+		public void SimpleCycle()
+		{
+			Clear();
+			AddVertex(0, 0, 1);
+			AddVertex(1, 0, 1);
+			AddExpectedComponent(0, 1);
+			ConfirmResults("Simple cycle");
+		}
+
+		[TestMethod]
+		public void ShortCycleInALongCycle()
+		{
+			Clear();
+			AddVertex(0, 1);
+			AddVertex(1, 0, 1, 2);
+			AddVertex(2, 2, 0);
+			AddExpectedComponent(0, 1, 2);
+			ConfirmResults("Short cycle in a long cycle");
+		}
+
+		[TestMethod]
+		public void CycleInRootDependingOnChain()
+		{
+			Clear();
+			AddVertex(0, 1);
+			AddVertex(1, 0, 2);
+			AddVertex(2, 3);
+			AddVertex(3);
+			AddExpectedComponent(0, 1);
+			ConfirmResults("Cycle in the root and with the root cycle depending on a chain of non-cycle verts");
+		}
+
+		[TestMethod]
+		public void CycleInRootDependingOnChainReverse()
+		{
+			Clear();
+			AddVertex(0);
+			AddVertex(1, 0);
+			AddVertex(2, 1, 3);
+			AddVertex(3, 2);
+			AddExpectedComponent(2, 3);
+			ConfirmResults("Cycle in the root and with the root cycle depending on a chain of non-cycle verts, submitted in reverse");
+		}
+
+		[TestMethod]
+		public void CycleAtLeafWithChainFromRoot()
+		{
+			Clear();
+			AddVertex(0, 1);
+			AddVertex(1, 2);
+			AddVertex(2, 3);
+			AddVertex(3, 2);
+			AddExpectedComponent(2, 3);
+			ConfirmResults("Cycle at a leaf and a chain from the root depending on that cycle");
+		}
+
+		[TestMethod]
+		public void CycleInRootWithChainDependents()
+		{
+			Clear();
+			AddVertex(0, 1);
+			AddVertex(1, 0);
+			AddVertex(2, 1);
+			AddVertex(3, 2);
+			AddExpectedComponent(0, 1);
+			ConfirmResults("Cycle in the root and with the root cycle depending on a chain of non-cycle verts");
+		}
+
+		[TestMethod]
+		public void VertexDependentUponACycle()
+		{
+			Clear();
+			AddVertex(0, 1);
+			AddVertex(1, 2, 3);
+			AddVertex(2, 1, 3);
+			AddVertex(3, 1, 2);
+			AddExpectedComponent(1, 2, 3);
+			ConfirmResults("Vertex dependent upon a cycle");
+		}
+
+		[TestMethod]
+		public void OneCycleDependentUponAnother()
+		{
+			// 0 -> 1 -> 2 -> 3 -> 1
+			//           |
+			//           v
+			//      4 -> 5 -> 6 -> 4
+			Clear();
+			AddVertex(0, 1);
+			AddVertex(1, 2);
+			AddVertex(2, 3, 5);
+			AddVertex(3, 1);
+			AddVertex(4, 5);
+			AddVertex(5, 6);
+			AddVertex(6, 4);
+			AddExpectedComponent(1, 2, 3);
+			AddExpectedComponent(4, 5, 6);
+			ConfirmResults("One cycle dependent upon another");
+		}
+
+		[TestMethod]
+		public void MutuallyReachableSetProblem1()
+		{
+			// 5 -> 0 -> (1 -> 2 -> 1)
+			//      |          |
+			//      |          v
+			//      |          5
+			//      v
+			//     (3 -> 4 -> 3)
+			Clear();
+			AddVertex(0, 1, 3);
+			AddVertex(1, 2);
+			AddVertex(2, 1, 5);
+			AddVertex(3, 4);
+			AddVertex(4, 3);
+			AddVertex(5, 0);
+			AddExpectedComponent(0, 1, 2, 5);
+			AddExpectedComponent(3, 4);
+			ConfirmResults("MutuallyReachableSet Problem1");
+		}
+
+		[TestMethod]
+		public void FullyConnectedGraph()
+		{
+			Clear();
+			List<int> allButItself = [0, 1, 2, 3, 4, 5];
+			int[] all = [0, 1, 2, 3, 4, 5];
+			for (int vertex = 0; vertex < 6; ++vertex)
+			{
+				allButItself.Remove(vertex);
+				AddVertex(vertex, [.. allButItself]);
+				allButItself.Add(vertex);
+			}
+			AddExpectedComponent(all);
+			ConfirmResults("FullyConnectedGraph");
+		}
+
+		[TestMethod]
+		public void CycleOfCycles()
+		{
+			// 0 - 1 - 2 - 0
+			// |
+			// 3 - 4 - 5 - 3
+			//         |
+			// 6 - 7 - 8 - 6
+			//     |
+			//     1
+			Clear();
+			AddVertex(0, 1, 3);
+			AddVertex(1, 2);
+			AddVertex(2, 0);
+			AddVertex(3, 4);
+			AddVertex(4, 5);
+			AddVertex(5, 3, 8);
+			AddVertex(6, 7);
+			AddVertex(7, 1, 8);
+			AddVertex(8, 6);
+			AddExpectedComponent(0, 1, 2, 3, 4, 5, 6, 7, 8);
+			ConfirmResults("Cycle of cycles");
+		}
+
+		[TestMethod]
+		public void TreeOfCycles()
+		{
+			// 0 ----- 1 ----- 2 ----- 0
+			// |       |       |
+			// |       0       |
+			// |               |
+			// 5 - 4 - 3 - 5   6 - 7 - 6
+			//     |\              |
+			//     | 5             |
+			//     |               |
+			//     8 - 9 - 10 - 9  11
+			//     |   |   |
+			//     10  8   8
+			Clear();
+			AddVertex(0, 1, 5);
+			AddVertex(1, 0, 2);
+			AddVertex(2, 0, 6);
+			AddVertex(3, 5);
+			AddVertex(4, 3, 5, 8);
+			AddVertex(5, 4);
+			AddVertex(6, 7);
+			AddVertex(7, 6, 11);
+			AddVertex(8, 9, 10);
+			AddVertex(9, 8, 10);
+			AddVertex(10, 8, 9);
+			AddVertex(11);
+			AddExpectedComponent(0, 1, 2);
+			AddExpectedComponent(3, 4, 5);
+			AddExpectedComponent(6, 7);
+			AddExpectedComponent(8, 9, 10);
+			ConfirmResults("Tree of cycles");
+		}
+	}
+}
diff --git a/Engine/Source/Programs/UnrealBuildTool/System/CondensationGraphBuilder.cs b/Engine/Source/Programs/UnrealBuildTool/System/CondensationGraphBuilder.cs
new file mode 100644
--- /dev/null
+++ b/Engine/Source/Programs/UnrealBuildTool/System/CondensationGraphBuilder.cs
@@ -0,0 +1,360 @@
+// Copyright Epic Games, Inc. All Rights Reserved.
+
+using System;
+using System.Collections.Generic;
+using Microsoft.Extensions.Logging;
+
+namespace UnrealBuildTool
+{
+	/// <summary>
+	/// Builds a condensation graph for a directed graph using Kosaraju-Sharir's algorithm.
+	/// Each strongly connected component (SCC) in the input is collapsed into a single node,
+	/// producing a directed acyclic graph (DAG). Vertices that are not part of any cycle form
+	/// their own single-element SCCs.
+	///
+	/// Build() returns false if the input graph is already a DAG (no cycles found), or true if
+	/// cycles were detected and the condensation differs from the input. In both cases GetSCCs()
+	/// and GetNodeToSCCIndex() are fully populated. GetCondensedEdges() is only populated when
+	/// Build() returns true (matching the C++ TryConstructCondensationGraph contract).
+	///
+	/// The output SCCs are ordered so that every edge in the condensation graph points from a
+	/// lower SCC index to a higher one (root-to-leaf topological order).
+	/// </summary>
+	/// <typeparam name="T">Type of node objects</typeparam>
+	public class CondensationGraphBuilder<T> where T : notnull
+	{
+		/// <summary>
+		/// UBT's logger object. Used for diagnostic output.
+		/// </summary>
+		public ILogger? Logger = null;
+
+		/// <summary>
+		/// Functor returning a node's name. Used for diagnostic logging.
+		/// </summary>
+		public Func<T, string>? NodeToString = null;
+
+		/// <summary>
+		/// List of all graph nodes in insertion order.
+		/// </summary>
+		private readonly List<GraphNode> Nodes = new List<GraphNode>();
+
+		/// <summary>
+		/// Maps each user-provided node to its internal GraphNode.
+		/// </summary>
+		private readonly Dictionary<T, GraphNode> NodeMap = new Dictionary<T, GraphNode>();
+
+		/// <summary>
+		/// Output: list of strongly connected components. Each SCC is an array of original nodes.
+		/// Ordered root-to-leaf: condensation edges always point from lower to higher SCC index.
+		/// Always populated after Build().
+		/// </summary>
+		private readonly List<T[]> ResultSCCs = new List<T[]>();
+
+		/// <summary>
+		/// Output: maps each original node to the index of its SCC in ResultSCCs.
+		/// Always populated after Build().
+		/// </summary>
+		private readonly Dictionary<T, int> ResultNodeToSCCIndex = new Dictionary<T, int>();
+
+		/// <summary>
+		/// Output: adjacency list of the condensation graph. ResultCondensedEdges[i] holds the
+		/// SCC indices that SCC i has outgoing edges to. Only populated when Build() returns true.
+		/// </summary>
+		private readonly List<int[]> ResultCondensedEdges = new List<int[]>();
+
+		/// <summary>
+		/// Initializes the builder and constructs the internal graph from the provided edges.
+		/// All nodes referenced by edges (both source and destination) are included.
+		/// Self-edges are discarded.
+		/// </summary>
+		/// <param name="Edges">Directed edges as (source, destination) tuples</param>
+		public CondensationGraphBuilder(List<Tuple<T, T>> Edges)
+		{
+			CreateGraph(Edges);
+		}
+
+		/// <summary>
+		/// Runs Kosaraju-Sharir's algorithm to compute the condensation graph.
+		/// </summary>
+		/// <returns>
+		/// True if cycles were found and the condensation differs from the input graph.
+		/// False if the input graph is already a DAG.
+		/// </returns>
+		public bool Build()
+		{
+			ResultSCCs.Clear();
+			ResultNodeToSCCIndex.Clear();
+			ResultCondensedEdges.Clear();
+
+			int NumVertices = Nodes.Count;
+			if (NumVertices == 0)
+			{
+				return false;
+			}
+
+			// Phase 1 — Iterative DFS on the original graph.
+			// Records each vertex in post-order (after all outgoing edges are exhausted).
+			// Any back-edge to an InProgress vertex (other than a self-edge) indicates a cycle.
+			List<int> PostOrder = new List<int>(NumVertices);
+			bool HasCycle = false;
+
+			{
+				VisitStatus[] Status = new VisitStatus[NumVertices]; // default: NotVisited
+				List<(int Vertex, int NextEdge)> Stack = new List<(int, int)>(NumVertices);
+
+				for (int Root = 0; Root < NumVertices; ++Root)
+				{
+					if (Status[Root] != VisitStatus.NotVisited)
+					{
+						continue;
+					}
+					Status[Root] = VisitStatus.InProgress;
+					Stack.Add((Root, 0));
+
+					while (Stack.Count > 0)
+					{
+						(int Vertex, int NextEdge) = Stack[Stack.Count - 1];
+						List<GraphNode> Edges = Nodes[Vertex].Links;
+						bool Pushed = false;
+
+						while (NextEdge < Edges.Count)
+						{
+							int EdgeVertex = Edges[NextEdge].Index;
+							++NextEdge;
+
+							if (Status[EdgeVertex] == VisitStatus.Visited)
+							{
+								// Already fully processed; nothing to do.
+							}
+							else if (Status[EdgeVertex] == VisitStatus.NotVisited)
+							{
+								Status[EdgeVertex] = VisitStatus.InProgress;
+								Stack[Stack.Count - 1] = (Vertex, NextEdge);
+								Stack.Add((EdgeVertex, 0));
+								Pushed = true;
+								break;
+							}
+							else // InProgress — back edge detected
+							{
+								// Self-edges are not counted as cycles; they are discarded during CreateGraph.
+								if (EdgeVertex != Vertex)
+								{
+									HasCycle = true;
+								}
+							}
+						}
+
+						if (!Pushed)
+						{
+							PostOrder.Add(Vertex);
+							Status[Vertex] = VisitStatus.Visited;
+							Stack.RemoveAt(Stack.Count - 1);
+						}
+					}
+				}
+			}
+
+			// No cycles found: input is already a DAG.
+			// Populate SCCs as an identity mapping in reverse post-order (root-to-leaf).
+			// ResultCondensedEdges is left empty per the C++ TryConstructCondensationGraph contract.
+			if (!HasCycle)
+			{
+				ResultSCCs.Capacity = NumVertices;
+				ResultNodeToSCCIndex.EnsureCapacity(NumVertices);
+
+				for (int SCCIndex = 0; SCCIndex < NumVertices; ++SCCIndex)
+				{
+					T NodeData = Nodes[PostOrder[NumVertices - 1 - SCCIndex]].Data;
+					ResultSCCs.Add(new T[] { NodeData });
+					ResultNodeToSCCIndex[NodeData] = SCCIndex;
+				}
+
+				return false;
+			}
+
+			// Phase 2 — Iterative DFS on the transpose graph in reverse post-order.
+			// Each unassigned vertex seeds a new SCC. All vertices reachable via TransposeLinks
+			// (i.e., all vertices that can reach the seed in the original graph) join the same SCC.
+			// Processing in reverse post-order ensures SCCs are discovered root-to-leaf.
+			{
+				bool[] Assigned = new bool[NumVertices];
+				List<(int Vertex, int NextEdge)> Stack = new List<(int, int)>(NumVertices);
+
+				for (int i = NumVertices - 1; i >= 0; --i)
+				{
+					int Root = PostOrder[i];
+					if (Assigned[Root])
+					{
+						continue;
+					}
+
+					List<int> SCCVertices = new List<int>();
+					Assigned[Root] = true;
+					SCCVertices.Add(Root);
+					Stack.Add((Root, 0));
+
+					while (Stack.Count > 0)
+					{
+						(int Vertex, int NextEdge) = Stack[Stack.Count - 1];
+						List<GraphNode> TransposeEdges = Nodes[Vertex].TransposeLinks;
+						bool Pushed = false;
+
+						while (NextEdge < TransposeEdges.Count)
+						{
+							int EdgeVertex = TransposeEdges[NextEdge].Index;
+							++NextEdge;
+
+							if (!Assigned[EdgeVertex])
+							{
+								Assigned[EdgeVertex] = true;
+								SCCVertices.Add(EdgeVertex);
+								Stack[Stack.Count - 1] = (Vertex, NextEdge);
+								Stack.Add((EdgeVertex, 0));
+								Pushed = true;
+								break;
+							}
+						}
+
+						if (!Pushed)
+						{
+							Stack.RemoveAt(Stack.Count - 1);
+						}
+					}
+
+					int SCCIndex = ResultSCCs.Count;
+					T[] SCCData = new T[SCCVertices.Count];
+					for (int j = 0; j < SCCVertices.Count; ++j)
+					{
+						T NodeData = Nodes[SCCVertices[j]].Data;
+						SCCData[j] = NodeData;
+						ResultNodeToSCCIndex[NodeData] = SCCIndex;
+					}
+					ResultSCCs.Add(SCCData);
+				}
+			}
+
+			// Phase 3 — Build condensation graph edges.
+			// For each SCC, union the outgoing edges of all its member nodes into the condensed
+			// edge list, mapping each target node to its SCC index and deduplicating with a
+			// boolean scratch array that is spot-cleared per iteration.
+			{
+				int NumSCCs = ResultSCCs.Count;
+				bool[] EdgeAdded = new bool[NumSCCs];
+
+				for (int SCCIndex = 0; SCCIndex < NumSCCs; ++SCCIndex)
+				{
+					List<int> SCCEdges = new List<int>();
+
+					foreach (T InNode in ResultSCCs[SCCIndex])
+					{
+						foreach (GraphNode Neighbor in NodeMap[InNode].Links)
+						{
+							int NeighborSCC = ResultNodeToSCCIndex[Neighbor.Data];
+							// NeighborSCC != SCCIndex also removes condensation self-edges
+							if (NeighborSCC != SCCIndex && !EdgeAdded[NeighborSCC])
+							{
+								EdgeAdded[NeighborSCC] = true;
+								SCCEdges.Add(NeighborSCC);
+							}
+						}
+					}
+
+					// Clear only the entries we touched — faster than resetting the whole array.
+					foreach (int EdgeSCC in SCCEdges)
+					{
+						EdgeAdded[EdgeSCC] = false;
+					}
+
+					ResultCondensedEdges.Add(SCCEdges.ToArray());
+				}
+			}
+
+			return true;
+		}
+
+		/// <summary>
+		/// Returns the list of strongly connected components (valid after calling Build()).
+		/// Each SCC is an array of original nodes. Ordered root-to-leaf in the condensation graph.
+		/// Always populated, even when Build() returns false.
+		/// </summary>
+		public List<T[]> GetSCCs() => ResultSCCs;
+
+		/// <summary>
+		/// Returns a mapping from each original node to the index of its SCC in GetSCCs()
+		/// (valid after calling Build()). Always populated, even when Build() returns false.
+		/// </summary>
+		public IReadOnlyDictionary<T, int> GetNodeToSCCIndex() => ResultNodeToSCCIndex;
+
+		/// <summary>
+		/// Returns the adjacency list of the condensation graph (valid after calling Build()).
+		/// GetCondensedEdges()[i] lists the SCC indices that SCC i has outgoing edges to.
+		/// Only populated when Build() returns true; empty when Build() returns false.
+		/// </summary>
+		public List<int[]> GetCondensedEdges() => ResultCondensedEdges;
+
+		/// <summary>
+		/// Constructs the internal graph from the provided edge list.
+		/// Nodes are added in the order they first appear. Self-edges are discarded.
+		/// </summary>
+		private void CreateGraph(List<Tuple<T, T>> Edges)
+		{
+			// First pass: register all nodes in insertion order.
+			foreach (Tuple<T, T> Edge in Edges)
+			{
+				if (!NodeMap.ContainsKey(Edge.Item1))
+				{
+					GraphNode Node = new GraphNode(Edge.Item1, Nodes.Count);
+					Nodes.Add(Node);
+					NodeMap[Edge.Item1] = Node;
+				}
+				if (!NodeMap.ContainsKey(Edge.Item2))
+				{
+					GraphNode Node = new GraphNode(Edge.Item2, Nodes.Count);
+					Nodes.Add(Node);
+					NodeMap[Edge.Item2] = Node;
+				}
+			}
+
+			// Second pass: add edges, discarding self-edges.
+			foreach (Tuple<T, T> Edge in Edges)
+			{
+				GraphNode Src = NodeMap[Edge.Item1];
+				GraphNode Dst = NodeMap[Edge.Item2];
+				if (Src == Dst)
+				{
+					continue;
+				}
+				Src.Links.Add(Dst);
+				Dst.TransposeLinks.Add(Src);
+			}
+		}
+
+		private enum VisitStatus : byte
+		{
+			NotVisited = 0,
+			InProgress = 1,
+			Visited    = 2,
+		}
+
+		private class GraphNode
+		{
+			/// <summary>User-provided node object.</summary>
+			public T Data;
+
+			/// <summary>Position in the Nodes list; used as an array index during Build().</summary>
+			public int Index;
+
+			/// <summary>Outgoing edges in the original graph.</summary>
+			public List<GraphNode> Links = new List<GraphNode>();
+
+			/// <summary>Incoming edges in the original graph (outgoing edges in the transpose graph).</summary>
+			public List<GraphNode> TransposeLinks = new List<GraphNode>();
+
+			public GraphNode(T Data, int Index)
+			{
+				this.Data  = Data;
+				this.Index = Index;
+			}
+		}
+	}
+}
