UnrealBuildTool: Fix static analysis warning in test.

#Synced-CL 52569840

diff --git a/Engine/Source/Programs/UnrealBuildTool.Tests/CppCompileWarningsTests.cs b/Engine/Source/Programs/UnrealBuildTool.Tests/CppCompileWarningsTests.cs
--- a/Engine/Source/Programs/UnrealBuildTool.Tests/CppCompileWarningsTests.cs
+++ b/Engine/Source/Programs/UnrealBuildTool.Tests/CppCompileWarningsTests.cs
@@ -343,7 +343,7 @@ namespace UnrealBuildTool.Tests
 					Assert.IsTrue(args.Count() == 1);
 					Assert.IsTrue(args.Contains("/w44702"));
 
-					CppCompileEnvironment preprocessOnlyCompileEnvironment = new CppCompileEnvironment(newCompileEnvironment);
+					CppCompileEnvironment preprocessOnlyCompileEnvironment = new(newCompileEnvironment);
 					preprocessOnlyCompileEnvironment.bPreprocessOnly = true;
 					args = baseWarnings.GenerateWarningCommandLineArgs(preprocessOnlyCompileEnvironment, typeof(VCToolChain));
 
